Copying //Tasks/UE4/Dev-Online-Plugins-v3 to //UE4/Dev-Main (Source: //Tasks/UE4/Dev-Online-Plugins-v3 @ 3047294)

diff --git a/Engine/Plugins/Online/OnlineSubsystemFacebook/Source/Private/IOS/OnlineSharingFacebook.cpp b/Engine/Plugins/Online/OnlineSubsystemFacebook/Source/Private/IOS/OnlineSharingFacebook.cpp
new file mode 100644
--- /dev/null
+++ b/Engine/Plugins/Online/OnlineSubsystemFacebook/Source/Private/IOS/OnlineSharingFacebook.cpp
@@ -0,0 +1,353 @@
+// Copyright 1998-2016 Epic Games, Inc. All Rights Reserved.
+
+
+// Module includes
+#include "OnlineSubsystemFacebookPrivatePCH.h"
+#include "OnlineSharingFacebook.h"
+#include "ImageCore.h"
+
+#import <FBSDKCoreKit/FBSDKCoreKit.h>
+#import <FBSDKShareKit/FBSDKShareKit.h>
+#import <FBSDKLoginKit/FBSDKLoginKit.h>
+
+FOnlineSharingFacebook::FOnlineSharingFacebook(class FOnlineSubsystemFacebook* InSubsystem)
+{
+	// Get our handle to the identity interface
+	IdentityInterface = InSubsystem->GetIdentityInterface();
+	SetupPermissionMaps();
+}
+
+
+FOnlineSharingFacebook::~FOnlineSharingFacebook()
+{
+	IdentityInterface = NULL;
+}
+
+
+void FOnlineSharingFacebook::SetupPermissionMaps()
+{
+	///////////////////////////////////////////////
+	// Read Permissions
+
+	ReadPermissionsMap.Add( EOnlineSharingReadCategory::Posts, [[NSArray alloc] 
+		initWithObjects:@"read_stream", nil] );
+
+	ReadPermissionsMap.Add( EOnlineSharingReadCategory::Friends, [[NSArray alloc]
+        initWithObjects:@"user_friends", nil] );
+
+	ReadPermissionsMap.Add( EOnlineSharingReadCategory::Mailbox, [[NSArray alloc] 
+		initWithObjects:@"read_mailbox", nil] );
+
+	ReadPermissionsMap.Add( EOnlineSharingReadCategory::OnlineStatus, [[NSArray alloc] 
+		initWithObjects:@"user_status", @"user_online_presence", nil] );
+    
+	ReadPermissionsMap.Add( EOnlineSharingReadCategory::ProfileInfo, [[NSArray alloc]
+        initWithObjects:@"public_profile", nil] );
+
+	ReadPermissionsMap.Add( EOnlineSharingReadCategory::LocationInfo, [[NSArray alloc] 
+		initWithObjects:@"user_checkins", @"user_hometown", nil] );
+
+
+	///////////////////////////////////////////////
+	// Publish Permissions
+	PublishPermissionsMap.Add( EOnlineSharingPublishingCategory::Posts, [[NSArray alloc] 
+		initWithObjects:@"publish_actions", nil] );
+	
+	PublishPermissionsMap.Add( EOnlineSharingPublishingCategory::Friends, [[NSArray alloc] 
+		initWithObjects:@"manage_friendlists", nil] );
+	
+	PublishPermissionsMap.Add( EOnlineSharingPublishingCategory::AccountAdmin, [[NSArray alloc] 
+		initWithObjects:@"manage_notifications", @"manage_pages", nil] );
+	
+	PublishPermissionsMap.Add( EOnlineSharingPublishingCategory::Events, [[NSArray alloc] 
+		initWithObjects:@"create_event", @"rsvp_event", nil ] );
+}
+
+
+bool FOnlineSharingFacebook::RequestNewReadPermissions(int32 LocalUserNum, EOnlineSharingReadCategory::Type NewPermissions)
+{
+	bool bTriggeredRequest = false;
+	if( IdentityInterface->GetLoginStatus(LocalUserNum) == ELoginStatus::LoggedIn )
+	{
+		bTriggeredRequest = true;
+
+		dispatch_async(dispatch_get_main_queue(),^ 
+			{
+				// Fill in an nsarray with the permissions which match those that the user has set,
+				// Here we iterate over each category, adding each individual permission linked with it in the ::SetupPermissionMaps
+				NSMutableArray* Permissions = [[NSMutableArray alloc] init];
+
+                
+                // Flag that dictates whether we need to reauthorize for permissions
+                bool bRequiresPermissionRequest = false;
+                
+				for(FReadPermissionsMap::TIterator It(ReadPermissionsMap); It; ++It)
+				{
+					if((NewPermissions & It.Key()) != 0)
+					{
+                        UE_LOG(LogOnline, Display, TEXT("ReadPermissionsMap[%i] - [%i]"), (int32)It.Key(), [It.Value() count]);
+                        for( NSString* RequestPermission in It.Value() )
+                        {
+							FBSDKAccessToken *accessToken = [FBSDKAccessToken currentAccessToken];
+							if (![accessToken.permissions containsObject:RequestPermission])
+							{
+                                bRequiresPermissionRequest = true;
+                                [Permissions addObject:RequestPermission];
+                            }
+                        }
+					}
+				}
+                
+                if( bRequiresPermissionRequest )
+                {
+					FBSDKLoginManager *loginManager = [[FBSDKLoginManager alloc] init];
+					[loginManager logInWithReadPermissions:Permissions
+												   handler:^(FBSDKLoginManagerLoginResult *result, NSError *error)
+						{
+							UE_LOG(LogOnline, Display, TEXT("logInWithReadPermissions : Success - %d"), error == nil);
+							TriggerOnRequestNewReadPermissionsCompleteDelegates(LocalUserNum, error==nil);
+						}
+					];
+                }
+                else
+                {
+                    // All permissions were already granted, no need to reauthorize
+                    TriggerOnRequestNewReadPermissionsCompleteDelegates(LocalUserNum, true);
+                }
+			}
+		);
+	}
+	else
+	{
+		// If we weren't logged into Facebook we cannot do this action
+		TriggerOnRequestNewReadPermissionsCompleteDelegates( LocalUserNum, false );
+	}
+
+	// We did kick off a request
+	return bTriggeredRequest;
+}
+
+
+bool FOnlineSharingFacebook::RequestNewPublishPermissions(int32 LocalUserNum, EOnlineSharingPublishingCategory::Type NewPermissions, EOnlineStatusUpdatePrivacy::Type Privacy)
+{
+	bool bTriggeredRequest = false;
+	if( IdentityInterface->GetLoginStatus(LocalUserNum) == ELoginStatus::LoggedIn )
+	{
+		bTriggeredRequest = true;
+
+		dispatch_async(dispatch_get_main_queue(),^ 
+			{
+				// Fill in an nsarray with the permissions which match those that the user has set,
+				// Here we iterate over each category, adding each individual permission linked with it in the ::SetupPermissionMaps
+				NSMutableArray* Permissions = [[NSMutableArray alloc] init];
+
+                // Flag that dictates whether we need to reauthorize for permissions
+                bool bRequiresPermissionRequest = false;
+                
+				for(FPublishPermissionsMap::TIterator It(PublishPermissionsMap); It; ++It)
+				{
+					if((NewPermissions & It.Key()) != 0)
+					{
+						UE_LOG(LogOnline, Verbose, TEXT("PublishPermissionsMap[%i] - [%i]"), (int32)It.Key(), [It.Value() count]);
+						for( NSString* RequestPermission in It.Value() )
+						{
+							FBSDKAccessToken *accessToken = [FBSDKAccessToken currentAccessToken];
+							if (![accessToken.permissions containsObject:RequestPermission])
+							{
+                                [Permissions addObject:RequestPermission];
+                                bRequiresPermissionRequest = true;
+                            }
+						}
+					}
+				}
+
+                if( bRequiresPermissionRequest )
+				{
+					FBSDKLoginManager *loginManager = [[FBSDKLoginManager alloc] init];
+                    FBSDKDefaultAudience DefaultAudience = FBSDKDefaultAudienceOnlyMe;
+                    switch (Privacy)
+                    {
+                        case EOnlineStatusUpdatePrivacy::OnlyMe:
+                            DefaultAudience = FBSDKDefaultAudienceOnlyMe;
+                            break;
+                        case EOnlineStatusUpdatePrivacy::OnlyFriends:
+                            DefaultAudience = FBSDKDefaultAudienceFriends;
+                            break;
+                        case EOnlineStatusUpdatePrivacy::Everyone:
+                            DefaultAudience = FBSDKDefaultAudienceEveryone;
+                            break;
+                    }
+					
+					[loginManager logInWithPublishPermissions:Permissions
+												   handler:^(FBSDKLoginManagerLoginResult *result, NSError *error)
+						{
+							UE_LOG(LogOnline, Display, TEXT("logInWithPublishPermissions : Success - %d"), error == nil);
+							TriggerOnRequestNewPublishPermissionsCompleteDelegates(LocalUserNum, error==nil);
+						}
+					 ];
+                }
+                else
+                {
+                    // All permissions were already granted, no need to reauthorize
+                    TriggerOnRequestNewPublishPermissionsCompleteDelegates(LocalUserNum, true);
+                }
+			}
+		);
+	}
+	else
+	{
+		// If we weren't logged into Facebook we cannot do this action
+		TriggerOnRequestNewPublishPermissionsCompleteDelegates( LocalUserNum, false );
+	}
+
+	// We did kick off a request
+	return bTriggeredRequest;
+}
+
+
+bool FOnlineSharingFacebook::ShareStatusUpdate(int32 LocalUserNum, const FOnlineStatusUpdate& StatusUpdate)
+{
+	bool bTriggeredRequest = false;
+	if( IdentityInterface->GetLoginStatus(LocalUserNum) == ELoginStatus::LoggedIn )
+	{
+		bTriggeredRequest = true;
+
+		dispatch_async(dispatch_get_main_queue(),^ 
+			{
+				// We need to determine where we are posting to.
+				NSString* GraphPath;
+
+				// The contents of the post
+				NSMutableDictionary* Params = [[NSMutableDictionary alloc] init];
+
+				NSString* ConvertedMessage = [NSString stringWithFString:StatusUpdate.Message];
+				[Params setObject:ConvertedMessage forKey:@"message"];
+
+				// Setup the image if one was added to the post
+				UIImage* SharingImage = nil;
+				if( StatusUpdate.Image != nil )
+				{
+					// We are posting this as a photo.
+					GraphPath = @"me/photos";
+
+					// Convert our FImage to a UIImage
+					int32 Width = StatusUpdate.Image->SizeX;
+					int32 Height = StatusUpdate.Image->SizeY;
+
+					CGColorSpaceRef colorSpace=CGColorSpaceCreateDeviceRGB();
+					CGContextRef bitmapContext=CGBitmapContextCreate(StatusUpdate.Image->RawData.GetData(), Width, Height, 8, 4*Width, colorSpace,  kCGImageAlphaNoneSkipLast);
+					CFRelease(colorSpace);
+
+					CGImageRef cgImage=CGBitmapContextCreateImage(bitmapContext);
+					CGContextRelease(bitmapContext);
+
+					SharingImage = [UIImage imageWithCGImage:cgImage];
+					CGImageRelease(cgImage);
+
+					[Params setObject:SharingImage forKey:@"picture"];
+				}
+				else
+				{
+					// if no image, we are posting this to the news feed.
+					GraphPath = @"me/feed";
+				}
+
+				// get the formatted friends tags
+				NSString* TaggedFriendIds = @"";
+
+				for(int32 FriendIdx = 0; FriendIdx < StatusUpdate.TaggedFriends.Num(); FriendIdx++)
+				{
+					NSString* FriendId = [NSString stringWithFString:StatusUpdate.TaggedFriends[FriendIdx]->ToString()];
+								
+					TaggedFriendIds = [TaggedFriendIds stringByAppendingFormat:@"%@%@", 
+						FriendId,
+						(FriendIdx < StatusUpdate.TaggedFriends.Num() - 1) ? @"," : @""
+					];
+				}
+
+				if( ![TaggedFriendIds isEqual:@""])
+				{
+					[Params setObject:TaggedFriendIds forKey:@"tags"];
+				}
+
+				// kick off a request to post the status
+				[[[FBSDKGraphRequest alloc]
+					initWithGraphPath:GraphPath
+					parameters:Params
+					HTTPMethod:@"POST"]
+					startWithCompletionHandler:^(FBSDKGraphRequestConnection *connection, id result, NSError *error)
+					{
+						UE_LOG(LogOnline, Display, TEXT("startWithGraphPath : Success - %d"), error == nil);
+						TriggerOnSharePostCompleteDelegates( LocalUserNum, error == nil );
+					}
+				];
+			}
+		);
+	}
+	else
+	{
+		// If we weren't logged into Facebook we cannot do this action
+		TriggerOnSharePostCompleteDelegates( LocalUserNum, false );
+	}
+
+	return bTriggeredRequest;
+}
+
+
+bool FOnlineSharingFacebook::ReadNewsFeed(int32 LocalUserNum, int32 NumPostsToRead)
+{
+	bool bTriggeredRequest = false;
+	if( IdentityInterface->GetLoginStatus(LocalUserNum) == ELoginStatus::LoggedIn )
+	{
+		bTriggeredRequest = true;
+		
+		dispatch_async(dispatch_get_main_queue(),^ 
+			{
+
+				// The current read permissions for this OSS.
+				EOnlineSharingReadCategory::Type ReadPermissions;
+				NSString *fqlQuery = 
+					[NSString stringWithFormat:@"SELECT post_id, created_time, type, attachment \
+												FROM stream WHERE filter_key in (SELECT filter_key \
+												FROM stream_filter WHERE uid=me() AND type='newsfeed')AND is_hidden = 0 \
+												LIMIT %d", NumPostsToRead];
+				
+				[[[FBSDKGraphRequest alloc]
+					initWithGraphPath:@"/fql"
+					parameters:[NSDictionary dictionaryWithObjectsAndKeys: fqlQuery, @"q", nil]
+					HTTPMethod:@"GET"]
+					startWithCompletionHandler:^(FBSDKGraphRequestConnection *connection, id result, NSError *error)
+					{
+						if( error )
+						{
+							UE_LOG(LogOnline, Display, TEXT("FTestSharingInterface::ReadStatusFeed - error[%s]"), *FString([error localizedDescription]));
+						}
+
+						TriggerOnReadNewsFeedCompleteDelegates(LocalUserNum, error==nil);
+					}
+				];
+			}
+		);
+	}
+	else
+	{
+		// If we weren't logged into Facebook we cannot do this action
+		TriggerOnReadNewsFeedCompleteDelegates( LocalUserNum, false );
+	}
+
+	return bTriggeredRequest;
+}
+
+
+EOnlineCachedResult::Type FOnlineSharingFacebook::GetCachedNewsFeed(int32 LocalUserNum, int32 NewsFeedIdx, FOnlineStatusUpdate& OutNewsFeed)
+{
+	check(NewsFeedIdx >= 0);
+	UE_LOG(LogOnline, Error, TEXT("FOnlineSharingFacebook::GetCachedNewsFeed not yet implemented"));
+	return EOnlineCachedResult::NotFound;
+}
+
+EOnlineCachedResult::Type FOnlineSharingFacebook::GetCachedNewsFeeds(int32 LocalUserNum, TArray<FOnlineStatusUpdate>& OutNewsFeeds)
+{
+	UE_LOG(LogOnline, Error, TEXT("FOnlineSharingFacebook::GetCachedNewsFeeds not yet implemented"));
+	return EOnlineCachedResult::NotFound;
+}
\ No newline at end of file
diff --git a/Engine/Plugins/Online/OnlineSubsystemFacebook/Source/Private/IOS/OnlineSharingFacebook.h b/Engine/Plugins/Online/OnlineSubsystemFacebook/Source/Private/IOS/OnlineSharingFacebook.h
new file mode 100644
--- /dev/null
+++ b/Engine/Plugins/Online/OnlineSubsystemFacebook/Source/Private/IOS/OnlineSharingFacebook.h
@@ -0,0 +1,65 @@
+// Copyright 1998-2016 Epic Games, Inc. All Rights Reserved.
+
+
+#pragma once
+ 
+
+// Module includes
+#include "OnlineSharingInterface.h"
+#include "OnlineSubsystemFacebookPackage.h"
+
+/**
+ * Facebook implementation of the Online Sharing Interface
+ */
+class FOnlineSharingFacebook : public IOnlineSharing
+{
+
+public:
+
+	//~ Begin IOnlineSharing Interface
+	virtual bool ReadNewsFeed(int32 LocalUserNum, int32 NumPostsToRead) override;
+	virtual bool RequestNewReadPermissions(int32 LocalUserNum, EOnlineSharingReadCategory::Type NewPermissions) override;
+	virtual bool ShareStatusUpdate(int32 LocalUserNum, const FOnlineStatusUpdate& StatusUpdate) override;
+	virtual bool RequestNewPublishPermissions(int32 LocalUserNum, EOnlineSharingPublishingCategory::Type NewPermissions, EOnlineStatusUpdatePrivacy::Type Privacy) override;
+	virtual EOnlineCachedResult::Type GetCachedNewsFeed(int32 LocalUserNum, int32 NewsFeedIdx, FOnlineStatusUpdate& OutNewsFeed) override;
+	virtual EOnlineCachedResult::Type GetCachedNewsFeeds(int32 LocalUserNum, TArray<FOnlineStatusUpdate>& OutNewsFeeds) override;
+	//~ End IOnlineSharing Interface
+
+	
+public:
+
+	/**
+	 * Constructor used to indicate which OSS we are a part of
+	 */
+	FOnlineSharingFacebook(class FOnlineSubsystemFacebook* InSubsystem);
+	
+	/**
+	 * Default destructor
+	 */
+	virtual ~FOnlineSharingFacebook();
+
+
+private:
+	
+	/**
+	 * Setup the permission categories with the relevant Facebook entries
+	 */
+	void SetupPermissionMaps();
+
+	// The read permissions map which sets up the facebook permissions in their correct category
+	typedef TMap< EOnlineSharingReadCategory::Type, NSArray* > FReadPermissionsMap;
+	FReadPermissionsMap ReadPermissionsMap;
+
+	// The publish permissions map which sets up the facebook permissions in their correct category
+	typedef TMap< EOnlineSharingPublishingCategory::Type, NSArray* > FPublishPermissionsMap;
+	FPublishPermissionsMap PublishPermissionsMap;
+
+
+private:
+
+	/** Reference to the main Facebook identity */
+	IOnlineIdentityPtr IdentityInterface;
+};
+
+
+typedef TSharedPtr<FOnlineSharingFacebook, ESPMode::ThreadSafe> FOnlineSharingFacebookPtr;
\ No newline at end of file
diff --git a/Engine/Plugins/Online/OnlineSubsystemFacebook/Source/Private/IOS/OnlineSubsystemFacebook.cpp b/Engine/Plugins/Online/OnlineSubsystemFacebook/Source/Private/IOS/OnlineSubsystemFacebook.cpp
new file mode 100644
--- /dev/null
+++ b/Engine/Plugins/Online/OnlineSubsystemFacebook/Source/Private/IOS/OnlineSubsystemFacebook.cpp
@@ -0,0 +1,216 @@
+// Copyright 1998-2016 Epic Games, Inc. All Rights Reserved.
+
+#include "OnlineSubsystemFacebookPrivatePCH.h"
+#include "IOSAppDelegate.h"
+
+#include "OnlineSharingFacebook.h"
+#include "OnlineUserFacebook.h"
+
+#import <FBSDKCoreKit/FBSDKCoreKit.h>
+
+FOnlineSubsystemFacebook::FOnlineSubsystemFacebook()
+	: FacebookIdentity(nullptr)
+	, FacebookFriends(nullptr)
+	, FacebookSharing(nullptr)
+	, FacebookUser(nullptr)
+{
+
+}
+
+FOnlineSubsystemFacebook::~FOnlineSubsystemFacebook()
+{
+	FacebookIdentity = nullptr;
+	FacebookFriends = nullptr;
+	FacebookSharing = nullptr; 
+	FacebookUser = nullptr;
+}
+
+IOnlineSessionPtr FOnlineSubsystemFacebook::GetSessionInterface() const
+{
+	return nullptr;
+}
+
+IOnlineIdentityPtr FOnlineSubsystemFacebook::GetIdentityInterface() const
+{
+	return FacebookIdentity;
+}
+
+IOnlineFriendsPtr FOnlineSubsystemFacebook::GetFriendsInterface() const
+{
+	return FacebookFriends;
+}
+
+IOnlineGroupsPtr FOnlineSubsystemFacebook::GetGroupsInterface() const
+{
+	return nullptr;
+}
+
+IOnlinePartyPtr FOnlineSubsystemFacebook::GetPartyInterface() const
+{
+	return nullptr;
+}
+
+IOnlineSharedCloudPtr FOnlineSubsystemFacebook::GetSharedCloudInterface() const
+{
+	return nullptr;
+}
+
+IOnlineUserCloudPtr FOnlineSubsystemFacebook::GetUserCloudInterface() const
+{
+	return nullptr;
+}
+
+IOnlineLeaderboardsPtr FOnlineSubsystemFacebook::GetLeaderboardsInterface() const
+{
+	return nullptr;
+}
+
+IOnlineVoicePtr FOnlineSubsystemFacebook::GetVoiceInterface() const
+{
+	return nullptr;
+}
+
+IOnlineExternalUIPtr FOnlineSubsystemFacebook::GetExternalUIInterface() const	
+{
+	return nullptr;
+}
+
+IOnlineTimePtr FOnlineSubsystemFacebook::GetTimeInterface() const
+{
+	return nullptr;
+}
+
+IOnlineTitleFilePtr FOnlineSubsystemFacebook::GetTitleFileInterface() const
+{
+	return nullptr;
+}
+
+IOnlineEntitlementsPtr FOnlineSubsystemFacebook::GetEntitlementsInterface() const
+{
+	return nullptr;
+}
+
+IOnlineStorePtr FOnlineSubsystemFacebook::GetStoreInterface() const
+{
+	return nullptr;
+}
+
+IOnlineEventsPtr FOnlineSubsystemFacebook::GetEventsInterface() const
+{
+	return nullptr;
+}
+
+IOnlineAchievementsPtr FOnlineSubsystemFacebook::GetAchievementsInterface() const
+{
+	return nullptr;
+}
+
+IOnlineSharingPtr FOnlineSubsystemFacebook::GetSharingInterface() const
+{
+	return FacebookSharing;
+}
+
+IOnlineUserPtr FOnlineSubsystemFacebook::GetUserInterface() const
+{
+	return FacebookUser;
+}
+
+IOnlineMessagePtr FOnlineSubsystemFacebook::GetMessageInterface() const
+{
+	return nullptr;
+}
+
+IOnlinePresencePtr FOnlineSubsystemFacebook::GetPresenceInterface() const
+{
+	return nullptr;
+}
+
+IOnlineChatPtr FOnlineSubsystemFacebook::GetChatInterface() const
+{
+	return nullptr;
+}
+
+IOnlineTurnBasedPtr FOnlineSubsystemFacebook::GetTurnBasedInterface() const
+{
+	return nullptr;
+}
+
+static void ListenFacebookOpenURL(UIApplication* application, NSURL* url, NSString* sourceApplication, id annotation)
+{
+	[[FBSDKApplicationDelegate sharedInstance] application:application
+												   openURL:url
+										 sourceApplication:sourceApplication
+												annotation:annotation];
+}
+
+bool FOnlineSubsystemFacebook::Init() 
+{
+	bool bSuccessfullyStartedUp = true;
+
+	[FBSDKAppEvents activateApp];
+	FIOSCoreDelegates::OnOpenURL.AddStatic(&ListenFacebookOpenURL);
+	
+	[[FBSDKApplicationDelegate sharedInstance] application:[UIApplication sharedApplication]
+							 didFinishLaunchingWithOptions:[IOSAppDelegate GetDelegate].launchOptions];
+	
+	FacebookIdentity = MakeShareable(new FOnlineIdentityFacebook());
+	FacebookSharing = MakeShareable(new FOnlineSharingFacebook(this));
+	FacebookFriends = MakeShareable(new FOnlineFriendsFacebook(this));
+	FacebookUser = MakeShareable(new FOnlineUserFacebook(this));
+	
+	return bSuccessfullyStartedUp;
+}
+
+bool FOnlineSubsystemFacebook::Shutdown() 
+{
+	bool bSuccessfullyShutdown = true;
+
+	bSuccessfullyShutdown = FOnlineSubsystemImpl::Shutdown();
+	return bSuccessfullyShutdown;
+}
+
+FString FOnlineSubsystemFacebook::GetAppId() const 
+{
+	return TEXT( "" );
+}
+
+bool FOnlineSubsystemFacebook::Exec(UWorld* InWorld, const TCHAR* Cmd, FOutputDevice& Ar)  
+{
+	if (FOnlineSubsystemImpl::Exec(InWorld, Cmd, Ar))
+	{
+		return true;
+	}
+	return false;
+}
+
+bool FOnlineSubsystemFacebook::Tick(float DeltaTime)
+{
+	if (!FOnlineSubsystemImpl::Tick(DeltaTime))
+	{
+		return false;
+	}
+
+	return true;
+}
+
+bool FOnlineSubsystemFacebook::IsEnabled()
+{
+	// Check the ini for disabling Facebook
+	bool bEnableFacebookSupport = false;
+	// Whilst we deprecate the old warning, let's still check if it's available
+	if (!GConfig->GetBool(TEXT("OnlineSubsystemFacebook"), TEXT("bEnabled"), bEnableFacebookSupport, GEngineIni))
+	{
+		GConfig->GetBool(TEXT("/Script/IOSRuntimeSettings.IOSRuntimeSettings"), TEXT("bEnableFacebookSupport"), bEnableFacebookSupport, GEngineIni);
+	}
+	else
+	{
+		UE_LOG(LogOnline, Warning, TEXT("The [OnlineSubsystemFacebook]:bEnabled flag has been moved to the IOS Project Settings, please update your preferences and App ID there."));
+	}
+
+#if !UE_BUILD_SHIPPING
+	// Check the commandline for disabling Facebook
+	bEnableFacebookSupport = bEnableFacebookSupport && !FParse::Param(FCommandLine::Get(), TEXT("NOFACEBOOK"));
+#endif
+
+	return bEnableFacebookSupport;
+}
diff --git a/Engine/Plugins/Online/OnlineSubsystemFacebook/Source/Private/IOS/OnlineUserFacebook.cpp b/Engine/Plugins/Online/OnlineSubsystemFacebook/Source/Private/IOS/OnlineUserFacebook.cpp
new file mode 100644
--- /dev/null
+++ b/Engine/Plugins/Online/OnlineSubsystemFacebook/Source/Private/IOS/OnlineUserFacebook.cpp
@@ -0,0 +1,228 @@
+// Copyright 1998-2016 Epic Games, Inc. All Rights Reserved.
+
+
+// Module includes
+#include "OnlineSubsystemFacebookPrivatePCH.h"
+#include "OnlineUserFacebook.h"
+
+#import <FBSDKCoreKit/FBSDKCoreKit.h>
+
+// FOnlineUserInfoFacebook
+
+TSharedRef<const FUniqueNetId> FOnlineUserInfoFacebook::GetUserId() const
+{
+	return UserId;
+}
+
+FString FOnlineUserInfoFacebook::GetRealName() const
+{
+	FString Result;
+	GetAccountData(TEXT("name"), Result);
+	return Result;
+}
+
+FString FOnlineUserInfoFacebook::GetDisplayName(const FString& Platform) const
+{
+	FString Result;
+	GetAccountData(TEXT("username"), Result);
+	return Result;
+}
+
+bool FOnlineUserInfoFacebook::GetUserAttribute(const FString& AttrName, FString& OutAttrValue) const
+{
+	return GetAccountData(AttrName, OutAttrValue);
+}
+
+// FOnlineUserFacebook
+
+FOnlineUserFacebook::FOnlineUserFacebook(class FOnlineSubsystemFacebook* InSubsystem)
+{
+	// Get our handle to the identity interface
+	IdentityInterface = (FOnlineIdentityFacebook*)InSubsystem->GetIdentityInterface().Get();
+}
+
+
+FOnlineUserFacebook::~FOnlineUserFacebook()
+{
+
+}
+
+bool FOnlineUserFacebook::QueryUserInfo(int32 LocalUserNum, const TArray<TSharedRef<const FUniqueNetId> >& UserIds)
+{
+	bool bTriggeredRequest = false;
+	
+	if(UserIds.Num() > 0 && IdentityInterface->GetLoginStatus(LocalUserNum) == ELoginStatus::LoggedIn)
+	{
+		bTriggeredRequest = true;
+		CachedUsers.Empty();
+
+		dispatch_async(dispatch_get_main_queue(),^ 
+			{
+				// We only ever have a single user in the facebook OSS.
+				FString ErrorStr;
+				bool bGatheredUserInfo = false;
+				if([FBSDKAccessToken currentAccessToken])
+				{
+					CachedUsers.Empty();
+					
+					bool bValidUserRequested = false;
+					const FString UserName([FBSDKProfile currentProfile].userID);
+					for(auto& NextUser : UserIds)
+					{
+						if(NextUser->ToString() == UserName)
+						{
+							bValidUserRequested = true;
+							break;
+						}
+					}
+					if(bValidUserRequested)
+					{
+						const FString RealName([FBSDKProfile currentProfile].name);
+						
+						TSharedRef<FOnlineUserInfoFacebook> FBUserInfo = MakeShareable(new FOnlineUserInfoFacebook(UserIds[0]->ToString()));
+						FBUserInfo->AccountData.Add(TEXT("name"), RealName);
+						FBUserInfo->AccountData.Add(TEXT("username"), UserName);
+						
+						UE_LOG(LogOnline, Display, TEXT("User Found: u:%s r:%s"), *UserName, *RealName);
+						CachedUsers.Add( FBUserInfo );
+						bGatheredUserInfo = true;
+					}
+					else
+					{
+						ErrorStr = TEXT("No user ids matched those of the single facebook user.");
+						UE_LOG(LogOnline, Display, TEXT("Failed to gather user information: %s"), *ErrorStr);
+					}
+
+				}
+				else
+				{
+					ErrorStr = TEXT("No valid login.");
+					UE_LOG(LogOnline, Display, TEXT("Failed to gather user information: %s"), *ErrorStr);
+				}
+				
+				TriggerOnQueryUserInfoCompleteDelegates(LocalUserNum, bGatheredUserInfo, UserIds, *ErrorStr);
+
+				
+/*
+				// A full list of all attainable data is here:
+				// https://developers.facebook.com/docs/reference/fql/user/
+
+				// Create a query with all the desired user ids and the data we wish to query
+				NSMutableString *fqlQuery = 
+					[NSMutableString stringWithFormat:@"SELECT name, username FROM user WHERE uid IN ("];
+
+				for( int32 Idx = 0; Idx < UserIds.Num(); Idx++ )
+				{
+					[fqlQuery appendString:[NSString stringWithFString:UserIds[Idx]->ToString()]];
+					
+					// Append a comma or close the parenthesis
+					if( Idx < UserIds.Num()-1 )
+					{
+						[fqlQuery appendString:@","];
+					}
+					else
+					{
+						[fqlQuery appendString:@")"];
+					}
+				}
+
+				UE_LOG(LogOnline, Display, TEXT("RunningFQL Query: %s"), *FString(fqlQuery));
+
+				// Kick off the FB Request
+				[[[FBSDKGraphRequest alloc]
+					initWithGraphPath:@"/fql"
+					parameters:[NSDictionary dictionaryWithObjectsAndKeys: fqlQuery, @"q", nil]
+					HTTPMethod:@"GET"]
+					startWithCompletionHandler:^(FBSDKGraphRequestConnection *connection, id result, NSError *error)
+					{
+						FString ErrorStr;
+						if( error )
+						{
+							ErrorStr = FString::Printf( TEXT("%s"), *FString([error localizedDescription]) );
+						}
+						else
+						{
+							NSArray* UserList = [[NSArray alloc] initWithArray:[result objectForKey:@"data"]];
+
+							// Clear our previously cached users before we repopulate the cache.
+							CachedUsers.Empty();
+
+							for( int32 UserIdx = 0; UserIdx < [UserList count]; UserIdx++ )
+							{	
+								NSDictionary* User = UserList[ UserIdx ];
+
+								const FString UserName([User objectForKey:@"username"]);
+								const FString RealName([User objectForKey:@"name"]);
+
+								TSharedRef<FOnlineUserInfoFacebook> FBUserInfo = MakeShareable(new FOnlineUserInfoFacebook(UserIds[UserIdx]->ToString()));
+								FBUserInfo->AccountData.Add(TEXT("name"), RealName);
+								FBUserInfo->AccountData.Add(TEXT("username"), UserName);
+								CachedUsers.Add( FBUserInfo );
+							}
+						}
+
+						TriggerOnQueryUserInfoCompleteDelegates(LocalUserNum, error == nil, UserIds, ErrorStr);
+					}
+				];
+ */
+			}
+		);
+	}
+	else
+	{
+		TriggerOnQueryUserInfoCompleteDelegates(LocalUserNum, false, UserIds, UserIds.Num() > 0 ? TEXT("Not logged in.") : TEXT("No users requested."));
+	}
+	
+	return bTriggeredRequest;
+}
+
+bool FOnlineUserFacebook::GetAllUserInfo(int32 LocalUserNum, TArray< TSharedRef<class FOnlineUser> >& OutUsers)
+{
+	UE_LOG(LogOnline, Verbose, TEXT("FOnlineUserFacebook::GetAllUserInfo()"));
+
+	for (int32 Idx=0; Idx < CachedUsers.Num(); Idx++)
+	{
+		OutUsers.Add(CachedUsers[Idx]);
+	}
+	return true;
+}
+
+TSharedPtr<FOnlineUser> FOnlineUserFacebook::GetUserInfo(int32 LocalUserNum, const class FUniqueNetId& UserId)
+{
+	TSharedPtr<FOnlineUser> Result;
+
+	UE_LOG(LogOnline, Verbose, TEXT("FOnlineUserFacebook::GetUserInfo()"));
+
+	for (int32 Idx=0; Idx < CachedUsers.Num(); Idx++)
+	{
+		if (*(CachedUsers[Idx]->GetUserId()) == UserId)
+		{
+			Result = CachedUsers[Idx];
+			break;
+		}
+	}
+	return Result;
+}
+
+bool FOnlineUserFacebook::QueryUserIdMapping(const FUniqueNetId& UserId, const FString& DisplayNameOrEmail, const FOnQueryUserMappingComplete& Delegate)
+{
+	Delegate.ExecuteIfBound(false, UserId, DisplayNameOrEmail, FUniqueNetIdString(), TEXT("not implemented"));
+	return false;
+}
+
+bool FOnlineUserFacebook::QueryExternalIdMappings(const FUniqueNetId& LocalUserId, const FExternalIdQueryOptions& QueryOptions, const TArray<FString>& ExternalIds, const FOnQueryExternalIdMappingsComplete& Delegate)
+{
+	Delegate.ExecuteIfBound(false, LocalUserId, QueryOptions, ExternalIds, TEXT("not implemented"));
+	return false;
+}
+
+void FOnlineUserFacebook::GetExternalIdMappings(const FExternalIdQueryOptions& QueryOptions, const TArray<FString>& ExternalIds, TArray<TSharedPtr<const FUniqueNetId>>& OutIds)
+{
+	// Not implemented for PS4 - return an array full of empty values
+	OutIds.SetNum(ExternalIds.Num());
+}
+
+TSharedPtr<const FUniqueNetId> FOnlineUserFacebook::GetExternalIdMapping(const FExternalIdQueryOptions& QueryOptions, const FString& ExternalId)
+{
+	return TSharedPtr<FUniqueNetId>();
+}
diff --git a/Engine/Plugins/Online/OnlineSubsystemFacebook/Source/Private/IOS/OnlineUserFacebook.h b/Engine/Plugins/Online/OnlineSubsystemFacebook/Source/Private/IOS/OnlineUserFacebook.h
new file mode 100644
--- /dev/null
+++ b/Engine/Plugins/Online/OnlineSubsystemFacebook/Source/Private/IOS/OnlineUserFacebook.h
@@ -0,0 +1,108 @@
+// Copyright 1998-2016 Epic Games, Inc. All Rights Reserved.
+
+
+#pragma once
+ 
+// Module includes
+#include "OnlineUserInterface.h"
+#include "OnlineSubsystemFacebookPackage.h"
+
+/**
+ * Info associated with an online user on the facebook service
+ */
+class FOnlineUserInfoFacebook : 
+	public FOnlineUser
+{
+public:
+
+	// FOnlineUser
+
+	virtual TSharedRef<const FUniqueNetId> GetUserId() const override;
+	virtual FString GetRealName() const override;
+	virtual FString GetDisplayName(const FString& Platform = FString()) const override;
+	virtual bool GetUserAttribute(const FString& AttrName, FString& OutAttrValue) const override;
+
+	// FOnlineUserInfoFacebook
+
+	/**
+	 * Init/default constructor
+	 */
+	FOnlineUserInfoFacebook(const FString& InUserId=TEXT("")) 
+		: UserId(new FUniqueNetIdString(InUserId))
+	{
+	}
+
+	/**
+	 * Destructor
+	 */
+	virtual ~FOnlineUserInfoFacebook()
+	{
+	}
+
+	/**
+	 * Get account data attribute
+	 *
+	 * @param Key account data entry key
+	 * @param OutVal [out] value that was found
+	 *
+	 * @return true if entry was found
+	 */
+	inline bool GetAccountData(const FString& Key, FString& OutVal) const
+	{
+		const FString* FoundVal = AccountData.Find(Key);
+		if (FoundVal != NULL)
+		{
+			OutVal = *FoundVal;
+			return true;
+		}
+		return false;
+	}
+
+	/** User Id represented as a FUniqueNetId */
+	TSharedRef<const FUniqueNetId> UserId;
+	/** Any addition account data associated with the friend */
+	TMap<FString, FString> AccountData;
+};
+
+
+/**
+ * Facebook implementation of the Online User Interface
+ */
+class FOnlineUserFacebook : public IOnlineUser
+{
+
+public:
+	
+	// IOnlineUser
+
+	virtual bool QueryUserInfo(int32 LocalUserNum, const TArray<TSharedRef<const FUniqueNetId> >& UserIds) override;
+	virtual bool GetAllUserInfo(int32 LocalUserNum, TArray< TSharedRef<class FOnlineUser> >& OutUsers) override;
+	virtual TSharedPtr<FOnlineUser> GetUserInfo(int32 LocalUserNum, const class FUniqueNetId& UserId) override;	
+	virtual bool QueryUserIdMapping(const FUniqueNetId& UserId, const FString& DisplayNameOrEmail, const FOnQueryUserMappingComplete& Delegate = FOnQueryUserMappingComplete()) override;
+	virtual bool QueryExternalIdMappings(const FUniqueNetId& LocalUserId, const FExternalIdQueryOptions& QueryOptions, const TArray<FString>& ExternalIds, const FOnQueryExternalIdMappingsComplete& Delegate = FOnQueryExternalIdMappingsComplete()) override;
+	virtual void GetExternalIdMappings(const FExternalIdQueryOptions& QueryOptions, const TArray<FString>& ExternalIds, TArray<TSharedPtr<const FUniqueNetId>>& OutIds) override;
+	virtual TSharedPtr<const FUniqueNetId> GetExternalIdMapping(const FExternalIdQueryOptions& QueryOptions, const FString& ExternalId) override;
+
+	// FOnlineUserFacebook
+
+	/**
+	 * Constructor used to indicate which OSS we are a part of
+	 */
+	FOnlineUserFacebook(class FOnlineSubsystemFacebook* InSubsystem);
+	
+	/**
+	 * Default destructor
+	 */
+	virtual ~FOnlineUserFacebook();
+
+
+private:
+
+	/** Reference to the main IOS identity */
+	class FOnlineIdentityFacebook* IdentityInterface;
+	/** The collection of facebook users received through the fb callbacks in QueryUserInfo */
+	TArray< TSharedRef<FOnlineUserInfoFacebook> > CachedUsers;
+};
+
+
+typedef TSharedPtr<FOnlineUserFacebook, ESPMode::ThreadSafe> FOnlineUserFacebookPtr;
\ No newline at end of file
diff --git a/Engine/Plugins/Online/OnlineSubsystemFacebook/Source/Private/OnlineSubsystemFacebookModule.cpp b/Engine/Plugins/Online/OnlineSubsystemFacebook/Source/Private/OnlineSubsystemFacebookModule.cpp
new file mode 100644
--- /dev/null
+++ b/Engine/Plugins/Online/OnlineSubsystemFacebook/Source/Private/OnlineSubsystemFacebookModule.cpp
@@ -0,0 +1,62 @@
+// Copyright 1998-2016 Epic Games, Inc. All Rights Reserved.
+
+#include "OnlineSubsystemFacebookPrivatePCH.h"
+
+// FOnlineSubsystemFacebookModule
+
+IMPLEMENT_MODULE(FOnlineSubsystemFacebookModule, OnlineSubsystemFacebook);
+
+/**
+ * Class responsible for creating instance(s) of the subsystem
+ */
+class FOnlineFactoryFacebook : public IOnlineFactory
+{
+public:
+
+	FOnlineFactoryFacebook() {}
+	virtual ~FOnlineFactoryFacebook() {}
+
+	virtual IOnlineSubsystemPtr CreateSubsystem(FName InstanceName)
+	{
+		FOnlineSubsystemFacebookPtr OnlineSub = MakeShareable(new FOnlineSubsystemFacebook());
+		if (OnlineSub->IsEnabled())
+		{
+			if(!OnlineSub->Init())
+			{
+				UE_LOG(LogOnline, Warning, TEXT("Facebook API failed to initialize!"));
+				OnlineSub->Shutdown();
+				OnlineSub = NULL;
+			}
+		}
+		else
+		{
+			UE_LOG(LogOnline, Warning, TEXT("Facebook API disabled!"));
+			OnlineSub->Shutdown();
+			OnlineSub = NULL;
+		}
+
+		return OnlineSub;
+	}
+};
+
+void FOnlineSubsystemFacebookModule::StartupModule()
+{
+	UE_LOG(LogOnline, Log, TEXT("Facebook Startup!"));
+
+	FacebookFactory = new FOnlineFactoryFacebook();
+
+	// Create and register our singleton factory with the main online subsystem for easy access
+	FOnlineSubsystemModule& OSS = FModuleManager::GetModuleChecked<FOnlineSubsystemModule>("OnlineSubsystem");
+	OSS.RegisterPlatformService(FACEBOOK_SUBSYSTEM, FacebookFactory);
+}
+
+void FOnlineSubsystemFacebookModule::ShutdownModule()
+{
+	UE_LOG(LogOnline, Log, TEXT("Facebook Shutdown!"));
+
+	FOnlineSubsystemModule& OSS = FModuleManager::GetModuleChecked<FOnlineSubsystemModule>("OnlineSubsystem");
+	OSS.UnregisterPlatformService(FACEBOOK_SUBSYSTEM);
+
+	delete FacebookFactory;
+	FacebookFactory = NULL;
+}
\ No newline at end of file
diff --git a/Engine/Plugins/Online/OnlineSubsystemFacebook/Source/Private/OnlineSubsystemFacebookPrivatePCH.h b/Engine/Plugins/Online/OnlineSubsystemFacebook/Source/Private/OnlineSubsystemFacebookPrivatePCH.h
new file mode 100644
--- /dev/null
+++ b/Engine/Plugins/Online/OnlineSubsystemFacebook/Source/Private/OnlineSubsystemFacebookPrivatePCH.h
@@ -0,0 +1,24 @@
+// Copyright 1998-2016 Epic Games, Inc. All Rights Reserved.
+
+#pragma once
+
+#include "Core.h"
+#include "ModuleManager.h"
+
+#include "Http.h"
+#include "Json.h"
+
+#include "OnlineSubsystemFacebook.h"
+#include "OnlineSubsystemFacebookModule.h"
+
+#if PLATFORM_WINDOWS
+#include "Windows/OnlineIdentityFacebook.h"
+#include "Windows/OnlineFriendsFacebook.h"
+#elif PLATFORM_IOS
+#include "IOS/OnlineIdentityFacebook.h"
+#include "IOS/OnlineFriendsFacebook.h"
+#endif
+
+/** pre-pended to all Facebook logging */
+#undef ONLINE_LOG_PREFIX
+#define ONLINE_LOG_PREFIX TEXT("Facebook: ")
\ No newline at end of file
diff --git a/Engine/Plugins/Online/OnlineSubsystemFacebook/Source/Private/Windows/OnlineFriendsFacebook.cpp b/Engine/Plugins/Online/OnlineSubsystemFacebook/Source/Private/Windows/OnlineFriendsFacebook.cpp
new file mode 100644
--- /dev/null
+++ b/Engine/Plugins/Online/OnlineSubsystemFacebook/Source/Private/Windows/OnlineFriendsFacebook.cpp
@@ -0,0 +1,329 @@
+// Copyright 1998-2016 Epic Games, Inc. All Rights Reserved.
+
+#include "OnlineSubsystemFacebookPrivatePCH.h"
+#include "OnlineIdentityFacebook.h"
+#include "OnlineFriendsFacebook.h"
+
+// FOnlineFriendFacebook
+
+TSharedRef<const FUniqueNetId> FOnlineFriendFacebook::GetUserId() const
+{
+	return UserId;
+}
+
+FString FOnlineFriendFacebook::GetRealName() const
+{
+	FString Result;
+	GetAccountData(TEXT("name"), Result);
+	return Result;
+}
+
+FString FOnlineFriendFacebook::GetDisplayName(const FString& Platform) const
+{
+	FString Result;
+	GetAccountData(TEXT("username"), Result);
+	return Result;
+}
+
+bool FOnlineFriendFacebook::GetUserAttribute(const FString& AttrName, FString& OutAttrValue) const
+{
+	return GetAccountData(AttrName, OutAttrValue);
+}
+
+EInviteStatus::Type FOnlineFriendFacebook::GetInviteStatus() const
+{
+	return EInviteStatus::Accepted;
+}
+
+const FOnlineUserPresence& FOnlineFriendFacebook::GetPresence() const
+{
+	return Presence;
+}
+
+// FOnlineFriendsFacebook
+
+FOnlineFriendsFacebook::FOnlineFriendsFacebook(FOnlineSubsystemFacebook* InSubsystem) 
+	: FacebookSubsystem(InSubsystem)
+{
+	check(FacebookSubsystem);
+
+	if (!GConfig->GetString(TEXT("OnlineSubsystemFacebook.OnlineFriendsFacebook"), TEXT("FriendsUrl"), FriendsUrl, GEngineIni))
+	{
+		UE_LOG(LogOnline, Warning, TEXT("Missing LoginUrl= in [OnlineSubsystemFacebook.OnlineIdentityFacebook] of DefaultEngine.ini"));
+	}
+	GConfig->GetArray(TEXT("OnlineSubsystemFacebook.OnlineFriendsFacebook"), TEXT("FriendsFields"), FriendsFields, GEngineIni);	
+	// always required fields
+	FriendsFields.AddUnique(TEXT("name"));
+	FriendsFields.AddUnique(TEXT("username"));
+}
+
+FOnlineFriendsFacebook::~FOnlineFriendsFacebook()
+{
+}
+
+bool FOnlineFriendsFacebook::ReadFriendsList(int32 LocalUserNum, const FString& ListName, const FOnReadFriendsListComplete& Delegate /*= FOnReadFriendsListComplete()*/)
+{
+	FString AccessToken;
+	FString ErrorStr;
+
+	if (!ListName.Equals(EFriendsLists::ToString(EFriendsLists::Default), ESearchCase::IgnoreCase))
+	{
+		ErrorStr = TEXT("Only the default friends list is supported");
+	}
+	// valid local player index
+	else if (LocalUserNum < 0 || LocalUserNum >= MAX_LOCAL_PLAYERS)
+	{
+		ErrorStr = FString::Printf(TEXT("Invalid LocalUserNum=%d"), LocalUserNum);
+	}
+	else
+	{
+		// Make sure a registration request for this user is not currently pending
+		for (TMap<IHttpRequest*, FPendingFriendsQuery>::TConstIterator It(FriendsQueryRequests); It; ++It)
+		{
+			const FPendingFriendsQuery& PendingFriendsQuery = It.Value();
+			if (PendingFriendsQuery.LocalUserNum == LocalUserNum)
+			{
+				ErrorStr = FString::Printf(TEXT("Already pending friends read for LocalUserNum=%d."), LocalUserNum);
+				break;
+			}
+		}
+		AccessToken = FacebookSubsystem->GetIdentityInterface()->GetAuthToken(LocalUserNum);
+		if (AccessToken.IsEmpty())
+		{
+			ErrorStr = FString::Printf(TEXT("Invalid access token for LocalUserNum=%d."), LocalUserNum);
+		}
+	}
+	if (!ErrorStr.IsEmpty())
+	{
+		UE_LOG_ONLINE(Warning, TEXT("ReadFriendsList request failed. %s"), *ErrorStr);
+		Delegate.ExecuteIfBound(LocalUserNum, false, ListName, ErrorStr);
+		return false;
+	}
+	
+	TSharedRef<IHttpRequest> HttpRequest = FHttpModule::Get().CreateRequest();
+	FriendsQueryRequests.Add(&HttpRequest.Get(), FPendingFriendsQuery(LocalUserNum));
+
+	// Optional list of fields to query for each friend
+	FString FieldsStr;
+	for (int32 Idx=0; Idx < FriendsFields.Num(); Idx++)
+	{
+		FieldsStr += FriendsFields[Idx];
+		if (Idx < (FriendsFields.Num()-1))
+		{
+			FieldsStr += TEXT(",");
+		}
+	}
+
+	// build the url
+	FString FriendsQueryUrl = FriendsUrl.Replace(TEXT("`fields"), *FieldsStr, ESearchCase::IgnoreCase);
+	FriendsQueryUrl = FriendsQueryUrl.Replace(TEXT("`token"), *AccessToken, ESearchCase::IgnoreCase);
+	
+	// kick off http request to read friends
+	HttpRequest->OnProcessRequestComplete().BindRaw(this, &FOnlineFriendsFacebook::QueryFriendsList_HttpRequestComplete, Delegate);
+	HttpRequest->SetURL(FriendsQueryUrl);
+	HttpRequest->SetHeader(TEXT("Content-Type"), TEXT("application/json"));
+	HttpRequest->SetVerb(TEXT("GET"));
+	return HttpRequest->ProcessRequest();
+}
+
+bool FOnlineFriendsFacebook::DeleteFriendsList(int32 LocalUserNum, const FString& ListName, const FOnDeleteFriendsListComplete& Delegate /*= FOnDeleteFriendsListComplete()*/)
+{
+	Delegate.ExecuteIfBound(LocalUserNum, false, ListName, FString(TEXT("DeleteFriendsList() is not supported")));
+	return false;
+}
+
+bool FOnlineFriendsFacebook::SendInvite(int32 LocalUserNum, const FUniqueNetId& FriendId, const FString& ListName, const FOnSendInviteComplete& Delegate /*= FOnSendInviteComplete()*/)
+{
+	Delegate.ExecuteIfBound(LocalUserNum, false, FriendId, ListName, FString(TEXT("SendInvite() is not supported")));
+	return false;
+}
+
+bool FOnlineFriendsFacebook::AcceptInvite(int32 LocalUserNum, const FUniqueNetId& FriendId, const FString& ListName, const FOnAcceptInviteComplete& Delegate /*= FOnAcceptInviteComplete()*/)
+{
+	Delegate.ExecuteIfBound(LocalUserNum, false, FriendId, ListName, FString(TEXT("AcceptInvite() is not supported")));
+	return false;
+}
+
+bool FOnlineFriendsFacebook::RejectInvite(int32 LocalUserNum, const FUniqueNetId& FriendId, const FString& ListName)
+{
+	TriggerOnRejectInviteCompleteDelegates(LocalUserNum, false, FriendId, ListName, FString(TEXT("RejectInvite() is not supported")));
+	return false;
+}
+
+bool FOnlineFriendsFacebook::DeleteFriend(int32 LocalUserNum, const FUniqueNetId& FriendId, const FString& ListName)
+{
+	TriggerOnDeleteFriendCompleteDelegates(LocalUserNum, false, FriendId, ListName, FString(TEXT("DeleteFriend() is not supported")));
+	return false;
+}
+
+bool FOnlineFriendsFacebook::GetFriendsList(int32 LocalUserNum, const FString& ListName, TArray< TSharedRef<FOnlineFriend> >& OutFriends)
+{
+	bool bResult = false;
+	// valid local player index
+	if (LocalUserNum >= 0 && LocalUserNum < MAX_LOCAL_PLAYERS)
+	{
+		// find friends list entry for local user
+		const FOnlineFriendsList* FriendsList = FriendsMap.Find(LocalUserNum);
+		if (FriendsList != NULL)
+		{
+			for (int32 FriendIdx=0; FriendIdx < FriendsList->Friends.Num(); FriendIdx++ )
+			{
+				OutFriends.Add(FriendsList->Friends[FriendIdx]);	
+			}
+			bResult = true;
+		}
+	}
+	return bResult;
+}
+
+TSharedPtr<FOnlineFriend> FOnlineFriendsFacebook::GetFriend(int32 LocalUserNum, const FUniqueNetId& FriendId, const FString& ListName)
+{
+	TSharedPtr<FOnlineFriend> Result;
+	// valid local player index
+	if (LocalUserNum >= 0 && LocalUserNum < MAX_LOCAL_PLAYERS)
+	{
+		// find friends list entry for local user
+		const FOnlineFriendsList* FriendsList = FriendsMap.Find(LocalUserNum);
+		if (FriendsList != NULL)
+		{
+			for (int32 FriendIdx=0; FriendIdx < FriendsList->Friends.Num(); FriendIdx++ )
+			{
+				if (*FriendsList->Friends[FriendIdx]->GetUserId() == FriendId)
+				{
+					Result = FriendsList->Friends[FriendIdx];
+					break;
+				}
+			}
+		}
+	}
+	return Result;
+}
+
+bool FOnlineFriendsFacebook::IsFriend(int32 LocalUserNum, const FUniqueNetId& FriendId, const FString& ListName)
+{
+	TSharedPtr<FOnlineFriend> Friend = GetFriend(LocalUserNum, FriendId, ListName);
+	if (Friend.IsValid() &&
+		Friend->GetInviteStatus() == EInviteStatus::Accepted)
+	{
+		return true;
+	}
+	return false;
+}
+
+bool FOnlineFriendsFacebook::QueryRecentPlayers(const FUniqueNetId& UserId, const FString& Namespace)
+{
+	UE_LOG(LogOnline, Verbose, TEXT("FOnlineFriendsFacebook::QueryRecentPlayers()"));
+
+	TriggerOnQueryRecentPlayersCompleteDelegates(UserId, Namespace, false, TEXT("not implemented"));
+
+	return false;
+}
+
+bool FOnlineFriendsFacebook::GetRecentPlayers(const FUniqueNetId& UserId, const FString& Namespace, TArray< TSharedRef<FOnlineRecentPlayer> >& OutRecentPlayers)
+{
+	return false;
+}
+
+bool FOnlineFriendsFacebook::BlockPlayer(int32 LocalUserNum, const FUniqueNetId& PlayerId)
+{
+	return false;
+}
+
+bool FOnlineFriendsFacebook::UnblockPlayer(int32 LocalUserNum, const FUniqueNetId& PlayerId)
+{
+	return false;
+}
+
+bool FOnlineFriendsFacebook::QueryBlockedPlayers(const FUniqueNetId& UserId)
+{
+	return false;
+}
+
+bool FOnlineFriendsFacebook::GetBlockedPlayers(const FUniqueNetId& UserId, TArray< TSharedRef<FOnlineBlockedPlayer> >& OutBlockedPlayers)
+{
+	return false;
+}
+
+void FOnlineFriendsFacebook::QueryFriendsList_HttpRequestComplete(FHttpRequestPtr HttpRequest, FHttpResponsePtr HttpResponse, bool bSucceeded, FOnReadFriendsListComplete Delegate)
+{
+	bool bResult = false;
+	FString ResponseStr, ErrorStr;
+
+	FPendingFriendsQuery PendingFriendsQuery = FriendsQueryRequests.FindRef(HttpRequest.Get());
+	// Remove the request from list of pending entries
+	FriendsQueryRequests.Remove(HttpRequest.Get());
+
+	if (bSucceeded &&
+		HttpResponse.IsValid())
+	{
+		ResponseStr = HttpResponse->GetContentAsString();
+		if (EHttpResponseCodes::IsOk(HttpResponse->GetResponseCode()))
+		{
+			UE_LOG(LogOnline, Verbose, TEXT("Query friends request complete. url=%s code=%d response=%s"), 
+				*HttpRequest->GetURL(), HttpResponse->GetResponseCode(), *ResponseStr);
+
+			// Create the Json parser
+			TSharedPtr<FJsonObject> JsonObject;
+			TSharedRef<TJsonReader<> > JsonReader = TJsonReaderFactory<>::Create(ResponseStr);
+
+			if (FJsonSerializer::Deserialize(JsonReader,JsonObject) &&
+				JsonObject.IsValid())
+			{
+				// Update cached entry for local user
+				FOnlineFriendsList& FriendsList = FriendsMap.FindOrAdd(PendingFriendsQuery.LocalUserNum);
+				FriendsList.Friends.Empty();
+
+				// Should have an array of id mappings
+				TArray<TSharedPtr<FJsonValue> > JsonFriends = JsonObject->GetArrayField(TEXT("data"));
+				for (TArray<TSharedPtr<FJsonValue> >::TConstIterator FriendIt(JsonFriends); FriendIt; ++FriendIt)
+				{
+					FString UserIdStr;
+					TMap<FString,FString> Attributes;
+					TSharedPtr<FJsonObject> JsonFriendEntry = (*FriendIt)->AsObject();
+					for (TMap<FString, TSharedPtr<FJsonValue > >::TConstIterator It(JsonFriendEntry->Values); It; ++It)
+					{
+						// parse user attributes
+						if (It->Value.IsValid() && 
+							It->Value->Type == EJson::String)
+						{
+							FString ValueStr = It->Value->AsString();
+							if (It->Key == TEXT("id"))
+							{
+								UserIdStr = ValueStr;
+							}
+							Attributes.Add(It->Key, ValueStr);
+						}
+					}
+					// only add if valid id
+					if (!UserIdStr.IsEmpty())
+					{
+						TSharedRef<FOnlineFriendFacebook> FriendEntry(new FOnlineFriendFacebook(UserIdStr));
+						FriendEntry->AccountData = Attributes;
+						// Add new friend entry to list
+						FriendsList.Friends.Add(FriendEntry);
+					}
+				}
+				bResult = true;
+			}
+		}
+		else
+		{
+			ErrorStr = FString::Printf(TEXT("Invalid response. code=%d error=%s"),
+				HttpResponse->GetResponseCode(), *ResponseStr);
+		}
+	}
+	else
+	{
+		ErrorStr = TEXT("No response");
+	}
+	if (!ErrorStr.IsEmpty())
+	{
+		UE_LOG(LogOnline, Warning, TEXT("Query friends list request failed. %s"), *ErrorStr);
+	}
+
+	Delegate.ExecuteIfBound(PendingFriendsQuery.LocalUserNum, bResult, EFriendsLists::ToString(EFriendsLists::Default), ErrorStr);
+}
+
+void FOnlineFriendsFacebook::DumpBlockedPlayers() const
+{
+}
\ No newline at end of file
diff --git a/Engine/Plugins/Online/OnlineSubsystemFacebook/Source/Private/Windows/OnlineFriendsFacebook.h b/Engine/Plugins/Online/OnlineSubsystemFacebook/Source/Private/Windows/OnlineFriendsFacebook.h
new file mode 100644
--- /dev/null
+++ b/Engine/Plugins/Online/OnlineSubsystemFacebook/Source/Private/Windows/OnlineFriendsFacebook.h
@@ -0,0 +1,156 @@
+// Copyright 1998-2016 Epic Games, Inc. All Rights Reserved.
+
+#pragma once
+
+#include "OnlineFriendsInterface.h"
+#include "OnlineSubsystemFacebookPackage.h"
+#include "OnlinePresenceInterface.h"
+
+/**
+ * Info associated with an online friend on the facebook service
+ */
+class FOnlineFriendFacebook : 
+	public FOnlineFriend
+{
+public:
+
+	// FOnlineUser
+
+	virtual TSharedRef<const FUniqueNetId> GetUserId() const override;
+	virtual FString GetRealName() const override;
+	virtual FString GetDisplayName(const FString& Platform = FString()) const override;
+	virtual bool GetUserAttribute(const FString& AttrName, FString& OutAttrValue) const override;
+
+	// FOnlineFriend
+	
+	virtual EInviteStatus::Type GetInviteStatus() const override;
+	virtual const FOnlineUserPresence& GetPresence() const override;
+
+	// FOnlineFriendMcp
+
+	/**
+	 * Init/default constructor
+	 */
+	FOnlineFriendFacebook(const FString& InUserId=TEXT("")) 
+		: UserId(new FUniqueNetIdString(InUserId))
+	{
+	}
+
+	/**
+	 * Destructor
+	 */
+	virtual ~FOnlineFriendFacebook()
+	{
+	}
+
+	/**
+	 * Get account data attribute
+	 *
+	 * @param Key account data entry key
+	 * @param OutVal [out] value that was found
+	 *
+	 * @return true if entry was found
+	 */
+	inline bool GetAccountData(const FString& Key, FString& OutVal) const
+	{
+		const FString* FoundVal = AccountData.Find(Key);
+		if (FoundVal != NULL)
+		{
+			OutVal = *FoundVal;
+			return true;
+		}
+		return false;
+	}
+
+	/** User Id represented as a FUniqueNetId */
+	TSharedRef<const FUniqueNetId> UserId;
+	/** Any addition account data associated with the friend */
+	TMap<FString, FString> AccountData;
+	/** @temp presence info */
+	FOnlineUserPresence Presence;
+};
+
+/**
+ * Facebook service implementation of the online friends interface
+ */
+class FOnlineFriendsFacebook :
+	public IOnlineFriends
+{	
+
+public:
+
+	// IOnlineFriends
+
+	virtual bool ReadFriendsList(int32 LocalUserNum, const FString& ListName, const FOnReadFriendsListComplete& Delegate = FOnReadFriendsListComplete()) override;
+	virtual bool DeleteFriendsList(int32 LocalUserNum, const FString& ListName, const FOnDeleteFriendsListComplete& Delegate = FOnDeleteFriendsListComplete()) override;
+	virtual bool SendInvite(int32 LocalUserNum, const FUniqueNetId& FriendId, const FString& ListName,  const FOnSendInviteComplete& Delegate = FOnSendInviteComplete()) override;
+	virtual bool AcceptInvite(int32 LocalUserNum, const FUniqueNetId& FriendId, const FString& ListName, const FOnAcceptInviteComplete& Delegate = FOnAcceptInviteComplete()) override;
+ 	virtual bool RejectInvite(int32 LocalUserNum, const FUniqueNetId& FriendId, const FString& ListName) override;
+ 	virtual bool DeleteFriend(int32 LocalUserNum, const FUniqueNetId& FriendId, const FString& ListName) override;
+	virtual bool GetFriendsList(int32 LocalUserNum, const FString& ListName, TArray< TSharedRef<FOnlineFriend> >& OutFriends) override;
+	virtual TSharedPtr<FOnlineFriend> GetFriend(int32 LocalUserNum, const FUniqueNetId& FriendId, const FString& ListName) override;
+	virtual bool IsFriend(int32 LocalUserNum, const FUniqueNetId& FriendId, const FString& ListName) override;
+	virtual bool QueryRecentPlayers(const FUniqueNetId& UserId, const FString& Namespace) override;
+	virtual bool GetRecentPlayers(const FUniqueNetId& UserId, const FString& Namespace, TArray< TSharedRef<FOnlineRecentPlayer> >& OutRecentPlayers) override;
+	virtual bool BlockPlayer(int32 LocalUserNum, const FUniqueNetId& PlayerId) override;
+	virtual bool UnblockPlayer(int32 LocalUserNum, const FUniqueNetId& PlayerId) override;
+	virtual bool QueryBlockedPlayers(const FUniqueNetId& UserId) override;
+	virtual bool GetBlockedPlayers(const FUniqueNetId& UserId, TArray< TSharedRef<FOnlineBlockedPlayer> >& OutBlockedPlayers) override;
+	virtual void DumpBlockedPlayers() const override;
+
+	// FOnlineFriendsFacebook
+
+	/**
+	 * Constructor
+	 *
+	 * @param InSubsystem Facebook subsystem being used
+	 */
+	FOnlineFriendsFacebook(class FOnlineSubsystemFacebook* InSubsystem);
+	
+	/**
+	 * Destructor
+	 */
+	virtual ~FOnlineFriendsFacebook();
+
+private:
+
+	/**
+	 * Should use the initialization constructor instead
+	 */
+	FOnlineFriendsFacebook();
+
+	/**
+	 * Delegate called when a user /me request from facebook is complete
+	 */
+	void QueryFriendsList_HttpRequestComplete(FHttpRequestPtr HttpRequest, FHttpResponsePtr HttpResponse, bool bSucceeded, FOnReadFriendsListComplete Delegate);
+
+	/** For accessing identity/token info of user logged in */
+	FOnlineSubsystemFacebook* FacebookSubsystem;
+	/** Config based url for querying friends list */
+	FString FriendsUrl;
+	/** Config based list of fields to use when querying friends list */
+	TArray<FString> FriendsFields;
+
+	/** List of online friends */
+	struct FOnlineFriendsList
+	{
+		TArray< TSharedRef<FOnlineFriendFacebook> > Friends;
+	};
+	/** Cached friends list from last call to ReadFriendsList for each local user */
+	TMap<int, FOnlineFriendsList> FriendsMap;
+
+	/** Info used to send request to register a user */
+	struct FPendingFriendsQuery
+	{
+		FPendingFriendsQuery(int32 InLocalUserNum=0) 
+			: LocalUserNum(InLocalUserNum)
+		{
+		}
+		/** local index of user making the request */
+		int32 LocalUserNum;
+	};
+	/** List of pending Http requests for user registration */
+	TMap<class IHttpRequest*, FPendingFriendsQuery> FriendsQueryRequests;
+};
+
+typedef TSharedPtr<FOnlineFriendsFacebook, ESPMode::ThreadSafe> FOnlineFriendsFacebookPtr;
diff --git a/Engine/Plugins/Online/OnlineSubsystemFacebook/Source/Private/Windows/OnlineIdentityFacebook.cpp b/Engine/Plugins/Online/OnlineSubsystemFacebook/Source/Private/Windows/OnlineIdentityFacebook.cpp
new file mode 100644
--- /dev/null
+++ b/Engine/Plugins/Online/OnlineSubsystemFacebook/Source/Private/Windows/OnlineIdentityFacebook.cpp
@@ -0,0 +1,461 @@
+// Copyright 1998-2016 Epic Games, Inc. All Rights Reserved.
+
+#include "OnlineSubsystemFacebookPrivatePCH.h"
+#include "OnlineIdentityFacebook.h"
+
+// FUserOnlineAccountFacebook
+
+TSharedRef<const FUniqueNetId> FUserOnlineAccountFacebook::GetUserId() const
+{
+	return UserIdPtr;
+}
+
+FString FUserOnlineAccountFacebook::GetRealName() const
+{
+	//@todo samz - implement
+	return FString();
+}
+
+FString FUserOnlineAccountFacebook::GetDisplayName(const FString& Platform) const
+{
+	//@todo samz - implement
+	return FString();
+}
+
+bool FUserOnlineAccountFacebook::GetUserAttribute(const FString& AttrName, FString& OutAttrValue) const
+{
+	return false;
+}
+
+bool FUserOnlineAccountFacebook::SetUserAttribute(const FString& AttrName, const FString& AttrValue)
+{
+	return false;
+}
+
+FString FUserOnlineAccountFacebook::GetAccessToken() const
+{
+	return AuthTicket;
+}
+
+bool FUserOnlineAccountFacebook::GetAuthAttribute(const FString& AttrName, FString& OutAttrValue) const
+{
+	return false;
+}
+
+// FOnlineIdentityFacebook
+
+/**
+ * Sets the needed configuration properties
+ */
+FOnlineIdentityFacebook::FOnlineIdentityFacebook()
+	: LastCheckElapsedTime(0.f)
+	, TotalCheckElapsedTime(0.f)
+	, MaxCheckElapsedTime(0.f)
+	, bHasLoginOutstanding(false)
+	, LocalUserNumPendingLogin(0)
+{
+	if (!GConfig->GetString(TEXT("OnlineSubsystemFacebook.OnlineIdentityFacebook"), TEXT("LoginUrl"), LoginUrl, GEngineIni))
+	{
+		UE_LOG(LogOnline, Warning, TEXT("Missing LoginUrl= in [OnlineSubsystemFacebook.OnlineIdentityFacebook] of DefaultEngine.ini"));
+	}
+	if (!GConfig->GetString(TEXT("OnlineSubsystemFacebook.OnlineIdentityFacebook"), TEXT("LoginRedirectUrl"), LoginRedirectUrl, GEngineIni))
+	{
+		UE_LOG(LogOnline, Warning, TEXT("Missing LoginRedirectUrl= in [OnlineSubsystemFacebook.OnlineIdentityFacebook] of DefaultEngine.ini"));
+	}
+	if (!GConfig->GetString(TEXT("OnlineSubsystemFacebook.OnlineIdentityFacebook"), TEXT("ClientId"), ClientId, GEngineIni))
+	{
+		UE_LOG(LogOnline, Warning, TEXT("Missing ClientId= in [OnlineSubsystemFacebook.OnlineIdentityFacebook] of DefaultEngine.ini"));
+	}
+	if (!GConfig->GetFloat(TEXT("OnlineSubsystemFacebook.OnlineIdentityFacebook"), TEXT("LoginTimeout"), MaxCheckElapsedTime, GEngineIni))
+	{
+		UE_LOG(LogOnline, Warning, TEXT("Missing LoginTimeout= in [OnlineSubsystemFacebook.OnlineIdentityFacebook] of DefaultEngine.ini"));
+		// Default to 30 seconds
+		MaxCheckElapsedTime = 30.f;
+	}
+}
+
+/**
+ * Used to do any time based processing of tasks
+ *
+ * @param DeltaTime the amount of time that has elapsed since the last tick
+ */
+void FOnlineIdentityFacebook::Tick(float DeltaTime)
+{
+	// Only tick once per frame
+	TickLogin(DeltaTime);
+}
+
+/**
+ * Ticks the registration process handling timeouts, etc.
+ *
+ * @param DeltaTime the amount of time that has elapsed since last tick
+ */
+void FOnlineIdentityFacebook::TickLogin(float DeltaTime)
+{
+	if (bHasLoginOutstanding)
+	{
+		LastCheckElapsedTime += DeltaTime;
+		TotalCheckElapsedTime += DeltaTime;
+		// See if enough time has elapsed in order to check for completion
+		if (LastCheckElapsedTime > 1.f ||
+			// Do one last check if we're getting ready to time out
+			TotalCheckElapsedTime > MaxCheckElapsedTime)
+		{
+			LastCheckElapsedTime = 0.f;
+			FString Title;
+			// Find the browser window we spawned which should now be titled with the redirect url
+			if (FPlatformMisc::GetWindowTitleMatchingText(*LoginRedirectUrl, Title))
+			{
+				bHasLoginOutstanding = false;
+
+				// Parse access token from the login redirect url
+				FString AccessToken;
+				if (FParse::Value(*Title, TEXT("access_token="), AccessToken) &&
+					!AccessToken.IsEmpty())
+				{
+					// strip off any url parameters and just keep the token itself
+					FString AccessTokenOnly;
+					if (AccessToken.Split(TEXT("&"), &AccessTokenOnly, NULL))
+					{
+						AccessToken = AccessTokenOnly;
+					}
+					// kick off http request to get user info with the new token
+					TSharedRef<IHttpRequest> HttpRequest = FHttpModule::Get().CreateRequest();
+					LoginUserRequests.Add(&HttpRequest.Get(), FPendingLoginUser(LocalUserNumPendingLogin, AccessToken));
+
+					FString MeUrl = TEXT("https://graph.facebook.com/me?access_token=`token");					
+
+					HttpRequest->OnProcessRequestComplete().BindRaw(this, &FOnlineIdentityFacebook::MeUser_HttpRequestComplete);
+					HttpRequest->SetURL(MeUrl.Replace(TEXT("`token"), *AccessToken, ESearchCase::IgnoreCase));
+					HttpRequest->SetHeader(TEXT("Content-Type"), TEXT("application/json"));
+					HttpRequest->SetVerb(TEXT("GET"));
+					HttpRequest->ProcessRequest();
+				}
+				else
+				{
+					TriggerOnLoginCompleteDelegates(LocalUserNumPendingLogin, false, FUniqueNetIdString(TEXT("")), 
+						FString(TEXT("RegisterUser() failed to parse the user registration results")));
+				}
+			}
+			// Trigger the delegate if we hit the timeout limit
+			else if (TotalCheckElapsedTime > MaxCheckElapsedTime)
+			{
+				bHasLoginOutstanding = false;
+				TriggerOnLoginCompleteDelegates(LocalUserNumPendingLogin, false, FUniqueNetIdString(TEXT("")), 
+					FString(TEXT("RegisterUser() timed out without getting the data")));
+			}
+		}
+		// Reset our time trackers if we are done ticking for now
+		if (!bHasLoginOutstanding)
+		{
+			LastCheckElapsedTime = 0.f;
+			TotalCheckElapsedTime = 0.f;
+		}
+	}
+}
+
+/**
+ * Parses the results into a user account entry
+ *
+ * @param Results the string returned by the login process
+ */
+bool FOnlineIdentityFacebook::ParseLoginResults(const FString& Results,FUserOnlineAccountFacebook& Account)
+{
+	// reset it
+	Account = FUserOnlineAccountFacebook();
+	// get the access token
+	if (FParse::Value(*Results, TEXT("access_token="), Account.AuthTicket))
+	{
+		return !Account.AuthTicket.IsEmpty();
+	}
+	return false;
+}
+
+/**
+ * Obtain online account info for a user that has been registered
+ *
+ * @param UserId user to search for
+ *
+ * @return info about the user if found, NULL otherwise
+ */
+TSharedPtr<FUserOnlineAccount> FOnlineIdentityFacebook::GetUserAccount(const FUniqueNetId& UserId) const
+{
+	TSharedPtr<FUserOnlineAccount> Result;
+
+	const TSharedRef<FUserOnlineAccountFacebook>* FoundUserAccount = UserAccounts.Find(UserId.ToString());
+	if (FoundUserAccount != NULL)
+	{
+		Result = *FoundUserAccount;
+	}
+
+	return Result;
+}
+
+/**
+ * Obtain list of all known/registered user accounts
+ *
+ * @param UserId user to search for
+ *
+ * @return info about the user if found, NULL otherwise
+ */
+TArray<TSharedPtr<FUserOnlineAccount> > FOnlineIdentityFacebook::GetAllUserAccounts() const
+{
+	TArray<TSharedPtr<FUserOnlineAccount> > Result;
+	
+	for (FUserOnlineAccountFacebookMap::TConstIterator It(UserAccounts); It; ++It)
+	{
+		Result.Add(It.Value());
+	}
+
+	return Result;
+}
+
+TSharedPtr<const FUniqueNetId> FOnlineIdentityFacebook::GetUniquePlayerId(int32 LocalUserNum) const
+{
+	const TSharedPtr<const FUniqueNetId>* FoundId = UserIds.Find(LocalUserNum);
+	if (FoundId != NULL)
+	{
+		return *FoundId;
+	}
+	return NULL;
+}
+
+/**
+ * Kicks off the browser process to register the user with Facebook
+ */
+bool FOnlineIdentityFacebook::Login(int32 LocalUserNum, const FOnlineAccountCredentials& AccountCredentials)
+{
+	FString ErrorStr;
+
+	if (bHasLoginOutstanding)
+	{
+		ErrorStr = FString::Printf(TEXT("Registration already pending for user %d"), 
+			LocalUserNumPendingLogin);
+	}
+	else if (!(LoginUrl.Len() && LoginRedirectUrl.Len() && ClientId.Len()))
+	{
+		ErrorStr = FString::Printf(TEXT("OnlineSubsystemFacebook is improperly configured in DefaultEngine.ini FacebookEndpoint=%s RedirectUrl=%s ClientId=%s"),
+			*LoginUrl, *LoginRedirectUrl, *ClientId);
+	}
+	else
+	{
+		// random number to represent client generated state for verification on login
+		State = FString::FromInt(FMath::Rand() % 100000);
+		// auth url to spawn in browser
+		const FString& Command = FString::Printf(TEXT("%s?redirect_uri=%s&client_id=%s&state=%s&response_type=token"), 
+			*LoginUrl, *LoginRedirectUrl, *ClientId, *State);
+		// This should open the browser with the command as the URL
+		if (FPlatformMisc::OsExecute(TEXT("open"), *Command))
+		{
+			// keep track of local user requesting registration
+			LocalUserNumPendingLogin = LocalUserNum;
+			bHasLoginOutstanding = true;
+		}
+		else
+		{
+			ErrorStr = FString::Printf(TEXT("Failed to execute command %s"),
+				*Command);
+		}
+	}
+
+	if (!ErrorStr.IsEmpty())
+	{
+		UE_LOG(LogOnline, Error, TEXT("RegisterUser() failed: %s"),
+			*ErrorStr);
+		TriggerOnLoginCompleteDelegates(LocalUserNum, false, FUniqueNetIdString(TEXT("")), ErrorStr);
+		return false;
+	}
+	return true;
+}
+
+TSharedPtr<const FUniqueNetId> FOnlineIdentityFacebook::CreateUniquePlayerId(uint8* Bytes, int32 Size)
+{
+	if (Bytes != NULL && Size > 0)
+	{
+		FString StrId(Size, (TCHAR*)Bytes);
+		return MakeShareable(new FUniqueNetIdString(StrId));
+	}
+	return NULL;
+}
+
+TSharedPtr<const FUniqueNetId> FOnlineIdentityFacebook::CreateUniquePlayerId(const FString& Str)
+{
+	return MakeShareable(new FUniqueNetIdString(Str));
+}
+
+// All of the methods below here fail because they aren't supported
+
+bool FOnlineIdentityFacebook::Logout(int32 LocalUserNum)
+{
+	TSharedPtr<const FUniqueNetId> UserId = GetUniquePlayerId(LocalUserNum);
+	if (UserId.IsValid())
+	{
+		// remove cached user account
+		UserAccounts.Remove(UserId->ToString());
+		// remove cached user id
+		UserIds.Remove(LocalUserNum);
+		// not async but should call completion delegate anyway
+		TriggerOnLogoutCompleteDelegates(LocalUserNum, true);
+
+		return true;
+	}
+	else
+	{
+		UE_LOG(LogOnline, Warning, TEXT("No logged in user found for LocalUserNum=%d."),
+			LocalUserNum);
+		TriggerOnLogoutCompleteDelegates(LocalUserNum, false);
+	}
+	return false;
+}
+
+bool FOnlineIdentityFacebook::AutoLogin(int32 LocalUserNum)
+{
+	return false;
+}
+
+ELoginStatus::Type FOnlineIdentityFacebook::GetLoginStatus(int32 LocalUserNum) const
+{
+	TSharedPtr<const FUniqueNetId> UserId = GetUniquePlayerId(LocalUserNum);
+	if (UserId.IsValid())
+	{
+		return GetLoginStatus(*UserId);
+	}
+	return ELoginStatus::NotLoggedIn;
+}
+
+ELoginStatus::Type FOnlineIdentityFacebook::GetLoginStatus(const FUniqueNetId& UserId) const
+{
+	TSharedPtr<FUserOnlineAccount> UserAccount = GetUserAccount(UserId);
+	if (UserAccount.IsValid() &&
+		UserAccount->GetUserId()->IsValid() &&
+		!UserAccount->GetAccessToken().IsEmpty())
+	{
+		return ELoginStatus::LoggedIn;
+	}
+	return ELoginStatus::NotLoggedIn;
+}
+
+FString FOnlineIdentityFacebook::GetPlayerNickname(int32 LocalUserNum) const
+{
+	TSharedPtr<const FUniqueNetId> UserId = GetUniquePlayerId(LocalUserNum);
+	if (UserId.IsValid())
+	{
+		return  GetPlayerNickname(*UserId);
+	}
+	return TEXT("");
+}
+
+FString FOnlineIdentityFacebook::GetPlayerNickname(const FUniqueNetId& UserId) const
+{
+	const TSharedRef<FUserOnlineAccountFacebook>* FoundUserAccount = UserAccounts.Find(UserId.ToString());
+	if (FoundUserAccount != NULL)
+	{
+		const TSharedRef<FUserOnlineAccountFacebook>& UserAccount = *FoundUserAccount;
+		return UserAccount->RealName;
+	}
+	return TEXT("");
+}
+
+FString FOnlineIdentityFacebook::GetAuthToken(int32 LocalUserNum) const
+{
+	TSharedPtr<const FUniqueNetId> UserId = GetUniquePlayerId(LocalUserNum);
+	if (UserId.IsValid())
+	{
+		TSharedPtr<FUserOnlineAccount> UserAccount = GetUserAccount(*UserId);
+		if (UserAccount.IsValid())
+		{
+			return UserAccount->GetAccessToken();
+		}
+	}
+	return FString();
+}
+
+void FOnlineIdentityFacebook::MeUser_HttpRequestComplete(FHttpRequestPtr HttpRequest, FHttpResponsePtr HttpResponse, bool bSucceeded)
+{
+	bool bResult = false;
+	FString ResponseStr, ErrorStr;
+	FUserOnlineAccountFacebook User;
+
+	FPendingLoginUser PendingRegisterUser = LoginUserRequests.FindRef(HttpRequest.Get());
+	// Remove the request from list of pending entries
+	LoginUserRequests.Remove(HttpRequest.Get());
+
+	if (bSucceeded &&
+		HttpResponse.IsValid())
+	{
+		ResponseStr = HttpResponse->GetContentAsString();
+		if (EHttpResponseCodes::IsOk(HttpResponse->GetResponseCode()))
+		{
+			UE_LOG(LogOnline, Verbose, TEXT("RegisterUser request complete. url=%s code=%d response=%s"), 
+				*HttpRequest->GetURL(), HttpResponse->GetResponseCode(), *ResponseStr);
+
+			if (User.FromJson(ResponseStr))
+			{
+				if (!User.UserId.IsEmpty())
+				{
+					// copy and construct the unique id
+					TSharedRef<FUserOnlineAccountFacebook> UserRef(new FUserOnlineAccountFacebook(User));
+					UserRef->UserIdPtr = MakeShareable(new FUniqueNetIdString(User.UserId));
+					// update/add cached entry for user
+					UserAccounts.Add(User.UserId, UserRef);
+					// update the access token
+					UserRef->AuthTicket = PendingRegisterUser.AccessToken;
+					// keep track of user ids for local users
+					UserIds.Add(PendingRegisterUser.LocalUserNum, UserRef->GetUserId());
+
+					bResult = true;
+				}
+				else
+				{
+					ErrorStr = FString::Printf(TEXT("Missing user id. payload=%s"),
+						*ResponseStr);
+				}
+			}
+			else
+			{
+				ErrorStr = FString::Printf(TEXT("Invalid response payload=%s"),
+					*ResponseStr);
+			}
+		}
+		else
+		{
+			ErrorStr = FString::Printf(TEXT("Invalid response. code=%d error=%s"),
+				HttpResponse->GetResponseCode(), *ResponseStr);
+		}
+	}
+	else
+	{
+		ErrorStr = TEXT("No response");
+	}
+	if (!ErrorStr.IsEmpty())
+	{
+		UE_LOG(LogOnline, Warning, TEXT("RegisterUser request failed. %s"), *ErrorStr);
+	}
+
+	TriggerOnLoginCompleteDelegates(PendingRegisterUser.LocalUserNum, bResult, FUniqueNetIdString(User.UserId), ErrorStr);
+}
+
+void FOnlineIdentityFacebook::GetUserPrivilege(const FUniqueNetId& UserId, EUserPrivileges::Type Privilege, const FOnGetUserPrivilegeCompleteDelegate& Delegate)
+{
+	Delegate.ExecuteIfBound(UserId, Privilege, (uint32)EPrivilegeResults::NoFailures);
+}	
+
+FPlatformUserId FOnlineIdentityFacebook::GetPlatformUserIdFromUniqueNetId(const FUniqueNetId& UniqueNetId)
+{
+	for (int i = 0; i < MAX_LOCAL_PLAYERS; ++i)
+	{
+		auto CurrentUniqueId = GetUniquePlayerId(i);
+		if (CurrentUniqueId.IsValid() && (*CurrentUniqueId == UniqueNetId))
+		{
+			return i;
+		}
+	}
+
+	return PLATFORMUSERID_NONE;
+}
+
+FString FOnlineIdentityFacebook::GetAuthType() const
+{
+	return TEXT("");
+}
+
diff --git a/Engine/Plugins/Online/OnlineSubsystemFacebook/Source/Private/Windows/OnlineIdentityFacebook.h b/Engine/Plugins/Online/OnlineSubsystemFacebook/Source/Private/Windows/OnlineIdentityFacebook.h
new file mode 100644
--- /dev/null
+++ b/Engine/Plugins/Online/OnlineSubsystemFacebook/Source/Private/Windows/OnlineIdentityFacebook.h
@@ -0,0 +1,184 @@
+// Copyright 1998-2016 Epic Games, Inc. All Rights Reserved.
+
+#pragma once
+ 
+#include "OnlineIdentityInterface.h"
+#include "OnlineSubsystemFacebookPackage.h"
+
+/**
+ * Info associated with an user account generated by this online service
+ */
+class FUserOnlineAccountFacebook : 
+	public FUserOnlineAccount,
+	public FOnlineJsonSerializable
+{
+public:
+
+	// FOnlineUser
+	
+	virtual TSharedRef<const FUniqueNetId> GetUserId() const override;
+	virtual FString GetRealName() const override;
+	virtual FString GetDisplayName(const FString& Platform = FString()) const override;
+	virtual bool GetUserAttribute(const FString& AttrName, FString& OutAttrValue) const override;
+	virtual bool SetUserAttribute(const FString& AttrName, const FString& AttrValue) override;
+
+	// FUserOnlineAccount
+
+	virtual FString GetAccessToken() const override;
+	virtual bool GetAuthAttribute(const FString& AttrName, FString& OutAttrValue) const override;
+
+	// FUserOnlineAccountFacebook
+
+	FUserOnlineAccountFacebook(const FString& InUserId=TEXT(""), const FString& InAuthTicket=TEXT("")) 
+		: UserIdPtr(new FUniqueNetIdString(InUserId))
+		, UserId(InUserId)
+		, AuthTicket(InAuthTicket)
+	{ }
+
+	virtual ~FUserOnlineAccountFacebook()
+	{
+	}
+
+	/** User Id represented as a FUniqueNetId */
+	TSharedRef<const FUniqueNetId> UserIdPtr;
+	/** Id associated with the user account provided by the online service during registration */
+	FString UserId;
+	/** Public user name */
+	FString UserName;
+	/** Real name */
+	FString RealName;
+	/** male or female */
+	FString Gender;
+	/** eg. en_US */
+	FString Locale;
+	/** Ticket which is provided to user once authenticated by the online service */
+	FString AuthTicket;
+
+	// FJsonSerializable
+
+	BEGIN_ONLINE_JSON_SERIALIZER
+		ONLINE_JSON_SERIALIZE("id", UserId);
+		ONLINE_JSON_SERIALIZE("username", UserName);
+		ONLINE_JSON_SERIALIZE("name", RealName);
+		ONLINE_JSON_SERIALIZE("gender", Gender);
+		ONLINE_JSON_SERIALIZE("locale", Locale);
+	END_ONLINE_JSON_SERIALIZER
+};
+/** Mapping from user id to his internal online account info (only one per user) */
+typedef TMap<FString, TSharedRef<FUserOnlineAccountFacebook> > FUserOnlineAccountFacebookMap;
+
+/**
+ * Facebook service implementation of the online identity interface
+ */
+class FOnlineIdentityFacebook :
+	public IOnlineIdentity
+{
+	/** The endpoint at Facebook we are supposed to hit for auth */
+	FString LoginUrl;
+	/** The redirect url for Facebook to redirect to upon completion. Note: this is configured at Facebook too */
+	FString LoginRedirectUrl;
+	/** The client id given to us by Facebook */
+	FString ClientId;
+
+	/** Users that have been registered/authenticated */
+	FUserOnlineAccountFacebookMap UserAccounts;
+	/** Ids mapped to locally registered users */
+	TMap<int32, TSharedPtr<const FUniqueNetId> > UserIds;
+
+	/** The amount of elapsed time since the last check */
+	float LastCheckElapsedTime;
+	/** Used to determine if we've timed out waiting for the response */
+	float TotalCheckElapsedTime;
+	/** Config value used to set our timeout period */
+	float MaxCheckElapsedTime;
+	/** Whether we have a registration in flight or not */
+	bool bHasLoginOutstanding;
+	/** A value used to verify our response came from our server */
+	FString State;
+	/** index of local user being registered */
+	int32 LocalUserNumPendingLogin;
+
+public:
+	// IOnlineIdentity
+	
+	virtual bool Login(int32 LocalUserNum, const FOnlineAccountCredentials& AccountCredentials) override;
+	virtual bool Logout(int32 LocalUserNum) override;
+	virtual bool AutoLogin(int32 LocalUserNum) override;
+	virtual TSharedPtr<FUserOnlineAccount> GetUserAccount(const FUniqueNetId& UserId) const override;
+	virtual TArray<TSharedPtr<FUserOnlineAccount> > GetAllUserAccounts() const override;
+	virtual TSharedPtr<const FUniqueNetId> GetUniquePlayerId(int32 LocalUserNum) const override;
+	virtual TSharedPtr<const FUniqueNetId> CreateUniquePlayerId(uint8* Bytes, int32 Size) override;
+	virtual TSharedPtr<const FUniqueNetId> CreateUniquePlayerId(const FString& Str) override;
+	virtual ELoginStatus::Type GetLoginStatus(int32 LocalUserNum) const override;
+	virtual ELoginStatus::Type GetLoginStatus(const FUniqueNetId& UserId) const override;
+	virtual FString GetPlayerNickname(int32 LocalUserNum) const override;
+	virtual FString GetPlayerNickname(const FUniqueNetId& UserId) const override;
+	virtual FString GetAuthToken(int32 LocalUserNum) const override;
+	virtual void GetUserPrivilege(const FUniqueNetId& UserId, EUserPrivileges::Type Privilege, const FOnGetUserPrivilegeCompleteDelegate& Delegate) override;
+	virtual FPlatformUserId GetPlatformUserIdFromUniqueNetId(const FUniqueNetId& UniqueNetId) override;
+	virtual FString GetAuthType() const override;
+
+	// FOnlineIdentityFacebook
+
+	FOnlineIdentityFacebook();
+
+	/**
+	 * Destructor
+	 */
+	virtual ~FOnlineIdentityFacebook()
+	{
+	}
+
+	/**
+	 * Used to do any time based processing of tasks
+	 *
+	 * @param DeltaTime the amount of time that has elapsed since the last tick
+	 */
+	void Tick(float DeltaTime);
+
+private:
+
+	/**
+	 * Ticks the registration process handling timeouts, etc.
+	 *
+	 * @param DeltaTime the amount of time that has elapsed since last tick
+	 */
+	void TickLogin(float DeltaTime);
+
+	/**
+	 * Parses the results into a user account entry
+	 *
+	 * @param Results the string returned by the login process
+	 * @param Account the account structure to fill in
+	 *
+	 * @return true if it parsed correctly, false otherwise
+	 */
+	bool ParseLoginResults(const FString& Results, FUserOnlineAccountFacebook& Account);
+
+	/**
+	 * Delegate called when a user /me request from facebook is complete
+	 */
+	void MeUser_HttpRequestComplete(FHttpRequestPtr HttpRequest, FHttpResponsePtr HttpResponse, bool bSucceeded);
+
+	/** Info used to send request to register a user */
+	struct FPendingLoginUser
+	{
+		FPendingLoginUser(
+			int32 InLocalUserNum=0, 
+			const FString& InAccessToken=FString()
+			) 
+			: LocalUserNum(InLocalUserNum)
+			, AccessToken(InAccessToken)
+		{
+
+		}
+		/** local index of user being registered */
+		int32 LocalUserNum;
+		/** Access token being used to login to Facebook */
+		FString AccessToken;
+	};
+	/** List of pending Http requests for user registration */
+	TMap<class IHttpRequest*, FPendingLoginUser> LoginUserRequests;
+};
+
+typedef TSharedPtr<FOnlineIdentityFacebook, ESPMode::ThreadSafe> FOnlineIdentityFacebookPtr;
\ No newline at end of file
diff --git a/Engine/Plugins/Online/OnlineSubsystemFacebook/Source/Private/Windows/OnlineSubsystemFacebook.cpp b/Engine/Plugins/Online/OnlineSubsystemFacebook/Source/Private/Windows/OnlineSubsystemFacebook.cpp
new file mode 100644
--- /dev/null
+++ b/Engine/Plugins/Online/OnlineSubsystemFacebook/Source/Private/Windows/OnlineSubsystemFacebook.cpp
@@ -0,0 +1,188 @@
+// Copyright 1998-2016 Epic Games, Inc. All Rights Reserved.
+
+#include "OnlineSubsystemFacebookPrivatePCH.h"
+#include "OnlineIdentityFacebook.h"
+#include "OnlineFriendsFacebook.h"
+
+// FOnlineSubsystemFacebook
+
+IOnlineSessionPtr FOnlineSubsystemFacebook::GetSessionInterface() const
+{
+	return nullptr;
+}
+
+IOnlineFriendsPtr FOnlineSubsystemFacebook::GetFriendsInterface() const
+{
+	return FacebookFriends;
+}
+
+IOnlinePartyPtr FOnlineSubsystemFacebook::GetPartyInterface() const
+{
+	return nullptr;
+}
+
+IOnlineGroupsPtr FOnlineSubsystemFacebook::GetGroupsInterface() const
+{
+	return nullptr;
+}
+
+IOnlineSharedCloudPtr FOnlineSubsystemFacebook::GetSharedCloudInterface() const
+{
+	return nullptr;
+}
+
+IOnlineUserCloudPtr FOnlineSubsystemFacebook::GetUserCloudInterface() const
+{
+	return nullptr;
+}
+
+IOnlineLeaderboardsPtr FOnlineSubsystemFacebook::GetLeaderboardsInterface() const
+{
+	return nullptr;
+}
+
+IOnlineVoicePtr FOnlineSubsystemFacebook::GetVoiceInterface() const
+{
+	return nullptr;
+}
+
+IOnlineExternalUIPtr FOnlineSubsystemFacebook::GetExternalUIInterface() const	
+{
+	return nullptr;
+}
+
+IOnlineTimePtr FOnlineSubsystemFacebook::GetTimeInterface() const
+{
+	return nullptr;
+}
+
+IOnlineIdentityPtr FOnlineSubsystemFacebook::GetIdentityInterface() const
+{
+	return FacebookIdentity;
+}
+
+IOnlineTitleFilePtr FOnlineSubsystemFacebook::GetTitleFileInterface() const
+{
+	return nullptr;
+}
+
+IOnlineEntitlementsPtr FOnlineSubsystemFacebook::GetEntitlementsInterface() const
+{
+	return nullptr;
+}
+
+IOnlineStorePtr FOnlineSubsystemFacebook::GetStoreInterface() const
+{
+	return nullptr;
+}
+
+IOnlineEventsPtr FOnlineSubsystemFacebook::GetEventsInterface() const
+{
+	return nullptr;
+}
+
+IOnlineAchievementsPtr FOnlineSubsystemFacebook::GetAchievementsInterface() const
+{
+	return nullptr;
+}
+
+IOnlineSharingPtr FOnlineSubsystemFacebook::GetSharingInterface() const
+{
+	return nullptr;
+}
+
+IOnlineUserPtr FOnlineSubsystemFacebook::GetUserInterface() const
+{
+	return nullptr;
+}
+
+IOnlineMessagePtr FOnlineSubsystemFacebook::GetMessageInterface() const
+{
+	return nullptr;
+}
+
+IOnlinePresencePtr FOnlineSubsystemFacebook::GetPresenceInterface() const
+{
+	return nullptr;
+}
+
+IOnlineChatPtr FOnlineSubsystemFacebook::GetChatInterface() const
+{
+	return nullptr;
+}
+
+IOnlineTurnBasedPtr FOnlineSubsystemFacebook::GetTurnBasedInterface() const
+{
+    return nullptr;
+}
+
+bool FOnlineSubsystemFacebook::Tick(float DeltaTime)
+{
+	if (!FOnlineSubsystemImpl::Tick(DeltaTime))
+	{
+		return false;
+	}
+
+	if (FacebookIdentity.IsValid())
+	{		
+		FacebookIdentity->Tick(DeltaTime);
+	}
+	return true;
+}
+
+bool FOnlineSubsystemFacebook::Init()
+{
+	FacebookIdentity = MakeShareable(new FOnlineIdentityFacebook());
+	FacebookFriends = MakeShareable(new FOnlineFriendsFacebook(this));
+	return true;
+}
+
+bool FOnlineSubsystemFacebook::Shutdown()
+{
+	UE_LOG(LogOnline, Display, TEXT("FOnlineSubsystemFacebook::Shutdown()"));
+	
+	FOnlineSubsystemImpl::Shutdown();
+
+	FacebookIdentity = nullptr;
+	FacebookFriends = nullptr;
+	return true;
+}
+
+FString FOnlineSubsystemFacebook::GetAppId() const
+{
+	return TEXT("Facebook");
+}
+
+bool FOnlineSubsystemFacebook::Exec(UWorld* InWorld, const TCHAR* Cmd, FOutputDevice& Ar) 
+{
+	if (FOnlineSubsystemImpl::Exec(InWorld, Cmd, Ar))
+	{
+		return true;
+	}
+	return false;
+}
+
+bool FOnlineSubsystemFacebook::IsEnabled()
+{
+	// Check the ini for disabling Facebook
+	bool bEnableFacebook = false;
+	GConfig->GetBool(TEXT("OnlineSubsystemFacebook"), TEXT("bEnabled"), bEnableFacebook, GEngineIni);
+
+#if !UE_BUILD_SHIPPING
+	// Check the commandline for disabling Facebook
+	bEnableFacebook = bEnableFacebook && !FParse::Param(FCommandLine::Get(),TEXT("NOFACEBOOK"));
+#endif
+
+	return bEnableFacebook;
+}
+
+FOnlineSubsystemFacebook::FOnlineSubsystemFacebook() 
+	: FacebookIdentity(nullptr)
+{
+
+}
+
+FOnlineSubsystemFacebook::~FOnlineSubsystemFacebook()
+{
+
+}
diff --git a/Engine/Plugins/Online/OnlineSubsystemFacebook/Source/Public/OnlineSubsystemFacebook.h b/Engine/Plugins/Online/OnlineSubsystemFacebook/Source/Public/OnlineSubsystemFacebook.h
new file mode 100644
--- /dev/null
+++ b/Engine/Plugins/Online/OnlineSubsystemFacebook/Source/Public/OnlineSubsystemFacebook.h
@@ -0,0 +1,91 @@
+// Copyright 1998-2016 Epic Games, Inc. All Rights Reserved.
+
+#pragma once
+
+#include "OnlineSubsystem.h"
+#include "OnlineSubsystemImpl.h"
+#include "OnlineSubsystemFacebookPackage.h"
+
+/** Forward declarations of all interface classes */
+typedef TSharedPtr<class FOnlineIdentityFacebook, ESPMode::ThreadSafe> FOnlineIdentityFacebookPtr;
+typedef TSharedPtr<class FOnlineFriendsFacebook, ESPMode::ThreadSafe> FOnlineFriendsFacebookPtr;
+typedef TSharedPtr<class FOnlineSharingFacebook, ESPMode::ThreadSafe> FOnlineSharingFacebookPtr;
+typedef TSharedPtr<class FOnlineUserFacebook, ESPMode::ThreadSafe> FOnlineUserFacebookPtr;
+
+/**
+ *	OnlineSubsystemFacebook - Implementation of the online subsystem for Facebook services
+ */
+class ONLINESUBSYSTEMFACEBOOK_API FOnlineSubsystemFacebook 
+	: public FOnlineSubsystemImpl
+{
+public:
+
+	// IOnlineSubsystem Interface
+	virtual IOnlineSessionPtr GetSessionInterface() const override;
+	virtual IOnlineFriendsPtr GetFriendsInterface() const override;
+	virtual IOnlinePartyPtr GetPartyInterface() const override;
+	virtual IOnlineGroupsPtr GetGroupsInterface() const override;
+	virtual IOnlineSharedCloudPtr GetSharedCloudInterface() const override;
+	virtual IOnlineUserCloudPtr GetUserCloudInterface() const override;
+	virtual IOnlineLeaderboardsPtr GetLeaderboardsInterface() const override;
+	virtual IOnlineVoicePtr GetVoiceInterface() const override;
+	virtual IOnlineExternalUIPtr GetExternalUIInterface() const override;
+	virtual IOnlineTimePtr GetTimeInterface() const override;
+	virtual IOnlineIdentityPtr GetIdentityInterface() const override;
+	virtual IOnlineTitleFilePtr GetTitleFileInterface() const override;
+	virtual IOnlineEntitlementsPtr GetEntitlementsInterface() const override;
+	virtual IOnlineStorePtr GetStoreInterface() const override;
+	virtual IOnlineStoreV2Ptr GetStoreV2Interface() const override { return nullptr; }
+	virtual IOnlinePurchasePtr GetPurchaseInterface() const override { return nullptr; }
+	virtual IOnlineEventsPtr GetEventsInterface() const override;
+	virtual IOnlineAchievementsPtr GetAchievementsInterface() const override;
+	virtual IOnlineSharingPtr GetSharingInterface() const override;
+	virtual IOnlineUserPtr GetUserInterface() const override;
+	virtual IOnlineMessagePtr GetMessageInterface() const override;
+	virtual IOnlinePresencePtr GetPresenceInterface() const override;
+	virtual IOnlineChatPtr GetChatInterface() const override;
+    virtual IOnlineTurnBasedPtr GetTurnBasedInterface() const override;
+	virtual bool Init() override;
+	virtual bool Shutdown() override;
+	virtual FString GetAppId() const override;
+	virtual bool Exec(class UWorld* InWorld, const TCHAR* Cmd, FOutputDevice& Ar) override;
+
+	// FTickerBaseObject
+
+	virtual bool Tick(float DeltaTime) override;
+
+	// FOnlineSubsystemFacebook
+
+	/**
+	 * Destructor
+	 */
+	virtual ~FOnlineSubsystemFacebook();
+
+	/**
+	 * Is Facebook available for use
+	 * @return true if Facebook functionality is available, false otherwise
+	 */
+	bool IsEnabled();
+
+PACKAGE_SCOPE:
+
+	/** Only the factory makes instances */
+	FOnlineSubsystemFacebook();
+
+private:
+
+	/** facebook implementation of identity interface */
+	FOnlineIdentityFacebookPtr FacebookIdentity;
+
+	/** facebook implementation of friends interface */
+	FOnlineFriendsFacebookPtr FacebookFriends;
+
+	/** facebook implementation of sharing interface */
+	FOnlineSharingFacebookPtr FacebookSharing;
+
+	/** facebook implementation of user interface */
+	FOnlineUserFacebookPtr FacebookUser;
+};
+
+typedef TSharedPtr<FOnlineSubsystemFacebook, ESPMode::ThreadSafe> FOnlineSubsystemFacebookPtr;
+
diff --git a/Engine/Plugins/Online/OnlineSubsystemNull/Source/OnlineSubsystemNull.Build.cs b/Engine/Plugins/Online/OnlineSubsystemNull/Source/OnlineSubsystemNull.Build.cs
new file mode 100644
--- /dev/null
+++ b/Engine/Plugins/Online/OnlineSubsystemNull/Source/OnlineSubsystemNull.Build.cs
@@ -0,0 +1,24 @@
+// Copyright 1998-2016 Epic Games, Inc. All Rights Reserved.
+
+using UnrealBuildTool;
+using System.IO;
+
+public class OnlineSubsystemNull : ModuleRules
+{
+	public OnlineSubsystemNull(TargetInfo Target)
+	{
+		Definitions.Add("ONLINESUBSYSTEMNULL_PACKAGE=1");
+
+		PrivateDependencyModuleNames.AddRange(
+			new string[] {
+				"Core", 
+				"CoreUObject", 
+				"Engine", 
+				"Sockets", 
+				"OnlineSubsystem", 
+				"OnlineSubsystemUtils",
+				"Json"
+			}
+			);
+	}
+}
diff --git a/Engine/Plugins/Online/OnlineSubsystemNull/Source/Private/NboSerializerNull.h b/Engine/Plugins/Online/OnlineSubsystemNull/Source/Private/NboSerializerNull.h
new file mode 100644
--- /dev/null
+++ b/Engine/Plugins/Online/OnlineSubsystemNull/Source/Private/NboSerializerNull.h
@@ -0,0 +1,82 @@
+// Copyright 1998-2016 Epic Games, Inc. All Rights Reserved.
+
+#pragma once
+
+#include "NboSerializer.h"
+#include "OnlineSubsystemNullTypes.h"
+
+/**
+ * Serializes data in network byte order form into a buffer
+ */
+class FNboSerializeToBufferNull : public FNboSerializeToBuffer
+{
+public:
+	/** Default constructor zeros num bytes*/
+	FNboSerializeToBufferNull() :
+		FNboSerializeToBuffer(512)
+	{
+	}
+
+	/** Constructor specifying the size to use */
+	FNboSerializeToBufferNull(uint32 Size) :
+		FNboSerializeToBuffer(Size)
+	{
+	}
+
+	/**
+	 * Adds Null session info to the buffer
+	 */
+ 	friend inline FNboSerializeToBufferNull& operator<<(FNboSerializeToBufferNull& Ar, const FOnlineSessionInfoNull& SessionInfo)
+ 	{
+		check(SessionInfo.HostAddr.IsValid());
+		// Skip SessionType (assigned at creation)
+		Ar << SessionInfo.SessionId;
+		Ar << *SessionInfo.HostAddr;
+		return Ar;
+ 	}
+
+	/**
+	 * Adds Null Unique Id to the buffer
+	 */
+	friend inline FNboSerializeToBufferNull& operator<<(FNboSerializeToBufferNull& Ar, const FUniqueNetIdString& UniqueId)
+	{
+		Ar << UniqueId.UniqueNetIdStr;
+		return Ar;
+	}
+};
+
+/**
+ * Class used to write data into packets for sending via system link
+ */
+class FNboSerializeFromBufferNull : public FNboSerializeFromBuffer
+{
+public:
+	/**
+	 * Initializes the buffer, size, and zeros the read offset
+	 */
+	FNboSerializeFromBufferNull(uint8* Packet,int32 Length) :
+		FNboSerializeFromBuffer(Packet,Length)
+	{
+	}
+
+	/**
+	 * Reads Null session info from the buffer
+	 */
+ 	friend inline FNboSerializeFromBufferNull& operator>>(FNboSerializeFromBufferNull& Ar, FOnlineSessionInfoNull& SessionInfo)
+ 	{
+		check(SessionInfo.HostAddr.IsValid());
+		// Skip SessionType (assigned at creation)
+		Ar >> SessionInfo.SessionId; 
+		Ar >> *SessionInfo.HostAddr;
+		return Ar;
+ 	}
+
+	/**
+	 * Reads Null Unique Id from the buffer
+	 */
+	friend inline FNboSerializeFromBufferNull& operator>>(FNboSerializeFromBufferNull& Ar, FUniqueNetIdString& UniqueId)
+	{
+		Ar >> UniqueId.UniqueNetIdStr;
+		return Ar;
+	}
+};
\ No newline at end of file
diff --git a/Engine/Plugins/Online/OnlineSubsystemNull/Source/Private/OnlineAsyncTaskManagerNull.cpp b/Engine/Plugins/Online/OnlineSubsystemNull/Source/Private/OnlineAsyncTaskManagerNull.cpp
new file mode 100644
--- /dev/null
+++ b/Engine/Plugins/Online/OnlineSubsystemNull/Source/Private/OnlineAsyncTaskManagerNull.cpp
@@ -0,0 +1,12 @@
+// Copyright 1998-2016 Epic Games, Inc. All Rights Reserved.
+
+#include "OnlineSubsystemNullPrivatePCH.h"
+#include "OnlineAsyncTaskManagerNull.h"
+#include "OnlineSubsystemNull.h"
+
+void FOnlineAsyncTaskManagerNull::OnlineTick()
+{
+	check(NullSubsystem);
+	check(FPlatformTLS::GetCurrentThreadId() == OnlineThreadId || !FPlatformProcess::SupportsMultithreading());
+}
+
diff --git a/Engine/Plugins/Online/OnlineSubsystemNull/Source/Private/OnlineIdentityNull.cpp b/Engine/Plugins/Online/OnlineSubsystemNull/Source/Private/OnlineIdentityNull.cpp
new file mode 100644
--- /dev/null
+++ b/Engine/Plugins/Online/OnlineSubsystemNull/Source/Private/OnlineIdentityNull.cpp
@@ -0,0 +1,311 @@
+// Copyright 1998-2016 Epic Games, Inc. All Rights Reserved.
+
+#include "OnlineSubsystemNullPrivatePCH.h"
+#include "OnlineIdentityNull.h"
+#include "IPAddress.h"
+#include "SocketSubsystem.h"
+
+bool FUserOnlineAccountNull::GetAuthAttribute(const FString& AttrName, FString& OutAttrValue) const
+{
+	const FString* FoundAttr = AdditionalAuthData.Find(AttrName);
+	if (FoundAttr != NULL)
+	{
+		OutAttrValue = *FoundAttr;
+		return true;
+	}
+	return false;
+}
+
+bool FUserOnlineAccountNull::GetUserAttribute(const FString& AttrName, FString& OutAttrValue) const
+{
+	const FString* FoundAttr = UserAttributes.Find(AttrName);
+	if (FoundAttr != NULL)
+	{
+		OutAttrValue = *FoundAttr;
+		return true;
+	}
+	return false;
+}
+
+bool FUserOnlineAccountNull::SetUserAttribute(const FString& AttrName, const FString& AttrValue)
+{
+	const FString* FoundAttr = UserAttributes.Find(AttrName);
+	if (FoundAttr == NULL || *FoundAttr != AttrValue)
+	{
+		UserAttributes.Add(AttrName, AttrValue);
+		return true;
+	}
+	return false;
+}
+
+inline FString GenerateRandomUserId(int32 LocalUserNum)
+{
+	FString HostName;
+	if (!ISocketSubsystem::Get()->GetHostName(HostName))
+	{
+		// could not get hostname, use address
+		bool bCanBindAll;
+		TSharedPtr<class FInternetAddr> Addr = ISocketSubsystem::Get()->GetLocalHostAddr(*GLog, bCanBindAll);
+		HostName = Addr->ToString(false);
+	}
+
+	const bool bForceUniqueId = FParse::Param( FCommandLine::Get(), TEXT( "StableNullID" ) );
+	
+	if ( ( GIsFirstInstance || bForceUniqueId ) && !GIsEditor )
+	{
+		// When possible, return a stable user id
+		return FString::Printf( TEXT( "%s-%s" ), *HostName, *FPlatformMisc::GetMachineId().ToString( EGuidFormats::Digits ) );
+	}
+
+	// If we're not the first instance (or in the editor), return truly random id
+	return FString::Printf( TEXT( "%s-%s" ), *HostName, *FGuid::NewGuid().ToString() );
+}
+
+bool FOnlineIdentityNull::Login(int32 LocalUserNum, const FOnlineAccountCredentials& AccountCredentials)
+{
+	FString ErrorStr;
+	TSharedPtr<FUserOnlineAccountNull> UserAccountPtr;
+	
+	// valid local player index
+	if (LocalUserNum < 0 || LocalUserNum >= MAX_LOCAL_PLAYERS)
+	{
+		ErrorStr = FString::Printf(TEXT("Invalid LocalUserNum=%d"), LocalUserNum);
+	}
+	else if (AccountCredentials.Id.IsEmpty())
+	{
+		ErrorStr = TEXT("Invalid account id, string empty");
+	}
+	else
+	{
+		TSharedPtr<const FUniqueNetId>* UserId = UserIds.Find(LocalUserNum);
+		if (UserId == NULL)
+		{
+			FString RandomUserId = GenerateRandomUserId(LocalUserNum);
+
+			FUniqueNetIdString NewUserId(RandomUserId);
+			UserAccountPtr = MakeShareable(new FUserOnlineAccountNull(RandomUserId));
+			UserAccountPtr->UserAttributes.Add(TEXT("id"), RandomUserId);
+
+			// update/add cached entry for user
+			UserAccounts.Add(NewUserId, UserAccountPtr.ToSharedRef());
+
+			// keep track of user ids for local users
+			UserIds.Add(LocalUserNum, UserAccountPtr->GetUserId());
+		}
+		else
+		{
+			const FUniqueNetIdString* UniqueIdStr = (FUniqueNetIdString*)(UserId->Get());
+			TSharedRef<FUserOnlineAccountNull>* TempPtr = UserAccounts.Find(*UniqueIdStr);
+			check(TempPtr);
+			UserAccountPtr = *TempPtr;
+		}
+	}
+
+	if (!ErrorStr.IsEmpty())
+	{
+		UE_LOG_ONLINE(Warning, TEXT("Login request failed. %s"), *ErrorStr);
+		TriggerOnLoginCompleteDelegates(LocalUserNum, false, FUniqueNetIdString(), ErrorStr);
+		return false;
+	}
+
+	TriggerOnLoginCompleteDelegates(LocalUserNum, true, *UserAccountPtr->GetUserId(), ErrorStr);
+	return true;
+}
+
+bool FOnlineIdentityNull::Logout(int32 LocalUserNum)
+{
+	TSharedPtr<const FUniqueNetId> UserId = GetUniquePlayerId(LocalUserNum);
+	if (UserId.IsValid())
+	{
+		// remove cached user account
+		UserAccounts.Remove(FUniqueNetIdString(*UserId));
+		// remove cached user id
+		UserIds.Remove(LocalUserNum);
+		// not async but should call completion delegate anyway
+		TriggerOnLogoutCompleteDelegates(LocalUserNum, true);
+
+		return true;
+	}
+	else
+	{
+		UE_LOG_ONLINE(Warning, TEXT("No logged in user found for LocalUserNum=%d."),
+			LocalUserNum);
+		TriggerOnLogoutCompleteDelegates(LocalUserNum, false);
+	}
+	return false;
+}
+
+bool FOnlineIdentityNull::AutoLogin(int32 LocalUserNum)
+{
+	FString LoginStr;
+	FString PasswordStr;
+	FString TypeStr;
+
+	FParse::Value(FCommandLine::Get(), TEXT("AUTH_LOGIN="), LoginStr);
+	FParse::Value(FCommandLine::Get(), TEXT("AUTH_PASSWORD="), PasswordStr);
+	FParse::Value(FCommandLine::Get(), TEXT("AUTH_TYPE="), TypeStr);
+	
+	if (!LoginStr.IsEmpty())
+	{
+		if (!PasswordStr.IsEmpty())
+		{
+			if (!TypeStr.IsEmpty())
+			{
+				return Login(0, FOnlineAccountCredentials(TypeStr, LoginStr, PasswordStr));
+			}
+			else
+			{
+				UE_LOG_ONLINE(Warning, TEXT("AutoLogin missing AUTH_TYPE=<type>."));
+			}
+		}
+		else
+		{
+			UE_LOG_ONLINE(Warning, TEXT("AutoLogin missing AUTH_PASSWORD=<password>."));
+		}
+	}
+	else
+	{
+		UE_LOG_ONLINE(Warning, TEXT("AutoLogin missing AUTH_LOGIN=<login id>."));
+	}
+	return false;
+}
+
+TSharedPtr<FUserOnlineAccount> FOnlineIdentityNull::GetUserAccount(const FUniqueNetId& UserId) const
+{
+	TSharedPtr<FUserOnlineAccount> Result;
+
+	FUniqueNetIdString StringUserId(UserId);
+	const TSharedRef<FUserOnlineAccountNull>* FoundUserAccount = UserAccounts.Find(StringUserId);
+	if (FoundUserAccount != NULL)
+	{
+		Result = *FoundUserAccount;
+	}
+
+	return Result;
+}
+
+TArray<TSharedPtr<FUserOnlineAccount> > FOnlineIdentityNull::GetAllUserAccounts() const
+{
+	TArray<TSharedPtr<FUserOnlineAccount> > Result;
+	
+	for (TMap<FUniqueNetIdString, TSharedRef<FUserOnlineAccountNull>>::TConstIterator It(UserAccounts); It; ++It)
+	{
+		Result.Add(It.Value());
+	}
+
+	return Result;
+}
+
+TSharedPtr<const FUniqueNetId> FOnlineIdentityNull::GetUniquePlayerId(int32 LocalUserNum) const
+{
+	const TSharedPtr<const FUniqueNetId>* FoundId = UserIds.Find(LocalUserNum);
+	if (FoundId != NULL)
+	{
+		return *FoundId;
+	}
+	return NULL;
+}
+
+TSharedPtr<const FUniqueNetId> FOnlineIdentityNull::CreateUniquePlayerId(uint8* Bytes, int32 Size)
+{
+	if (Bytes != NULL && Size > 0)
+	{
+		FString StrId(Size, (TCHAR*)Bytes);
+		return MakeShareable(new FUniqueNetIdString(StrId));
+	}
+	return NULL;
+}
+
+TSharedPtr<const FUniqueNetId> FOnlineIdentityNull::CreateUniquePlayerId(const FString& Str)
+{
+	return MakeShareable(new FUniqueNetIdString(Str));
+}
+
+ELoginStatus::Type FOnlineIdentityNull::GetLoginStatus(int32 LocalUserNum) const
+{
+	TSharedPtr<const FUniqueNetId> UserId = GetUniquePlayerId(LocalUserNum);
+	if (UserId.IsValid())
+	{
+		return GetLoginStatus(*UserId);
+	}
+	return ELoginStatus::NotLoggedIn;
+}
+
+ELoginStatus::Type FOnlineIdentityNull::GetLoginStatus(const FUniqueNetId& UserId) const 
+{
+	TSharedPtr<FUserOnlineAccount> UserAccount = GetUserAccount(UserId);
+	if (UserAccount.IsValid() &&
+		UserAccount->GetUserId()->IsValid())
+	{
+		return ELoginStatus::LoggedIn;
+	}
+	return ELoginStatus::NotLoggedIn;
+}
+
+FString FOnlineIdentityNull::GetPlayerNickname(int32 LocalUserNum) const
+{
+	TSharedPtr<const FUniqueNetId> UniqueId = GetUniquePlayerId(LocalUserNum);
+	if (UniqueId.IsValid())
+	{
+		return UniqueId->ToString();
+	}
+
+	return TEXT("NullUser");
+}
+
+FString FOnlineIdentityNull::GetPlayerNickname(const FUniqueNetId& UserId) const
+{
+	return UserId.ToString();
+}
+
+FString FOnlineIdentityNull::GetAuthToken(int32 LocalUserNum) const
+{
+	TSharedPtr<const FUniqueNetId> UserId = GetUniquePlayerId(LocalUserNum);
+	if (UserId.IsValid())
+	{
+		TSharedPtr<FUserOnlineAccount> UserAccount = GetUserAccount(*UserId);
+		if (UserAccount.IsValid())
+		{
+			return UserAccount->GetAccessToken();
+		}
+	}
+	return FString();
+}
+
+FOnlineIdentityNull::FOnlineIdentityNull(class FOnlineSubsystemNull* InSubsystem)
+{
+	// autologin the 0-th player
+	Login(0, FOnlineAccountCredentials(TEXT("DummyType"), TEXT("DummyUser"), TEXT("DummyId")) );
+}
+
+FOnlineIdentityNull::FOnlineIdentityNull()
+{
+}
+
+FOnlineIdentityNull::~FOnlineIdentityNull()
+{
+}
+
+void FOnlineIdentityNull::GetUserPrivilege(const FUniqueNetId& UserId, EUserPrivileges::Type Privilege, const FOnGetUserPrivilegeCompleteDelegate& Delegate)
+{
+	Delegate.ExecuteIfBound(UserId, Privilege, (uint32)EPrivilegeResults::NoFailures);
+}
+
+FPlatformUserId FOnlineIdentityNull::GetPlatformUserIdFromUniqueNetId(const FUniqueNetId& UniqueNetId)
+{
+	for (int i = 0; i < MAX_LOCAL_PLAYERS; ++i)
+	{
+		auto CurrentUniqueId = GetUniquePlayerId(i);
+		if (CurrentUniqueId.IsValid() && (*CurrentUniqueId == UniqueNetId))
+		{
+			return i;
+		}
+	}
+
+	return PLATFORMUSERID_NONE;
+}
+
+FString FOnlineIdentityNull::GetAuthType() const
+{
+	return TEXT("");
+}
