- Deleted unused OSSFacebook and OSSGoogle platform code

diff --git a/Engine/Plugins/Online/OnlineSubsystemFacebook/Source/OnlineSubsystemFacebook.Build.cs b/Engine/Plugins/Online/OnlineSubsystemFacebook/Source/OnlineSubsystemFacebook.Build.cs
--- a/Engine/Plugins/Online/OnlineSubsystemFacebook/Source/OnlineSubsystemFacebook.Build.cs
+++ b/Engine/Plugins/Online/OnlineSubsystemFacebook/Source/OnlineSubsystemFacebook.Build.cs
@@ -76,16 +76,6 @@ public class OnlineSubsystemFacebook : ModuleRules
 			bUsesRestfulImpl = true;
 			PublicDefinitions.Add("WITH_FACEBOOK=1");
 		}
-		else if (Target.Platform == UnrealTargetPlatform.XboxOne)
-		{
-			PublicDefinitions.Add("WITH_FACEBOOK=1");
-			PrivateIncludePaths.Add("Private/XboxOne");
-		}
-		else if (Target.Platform == UnrealTargetPlatform.PS4)
-		{
-			PublicDefinitions.Add("WITH_FACEBOOK=1");
-			PrivateIncludePaths.Add("Private/PS4");
-		}
 		else if (Target.Platform == UnrealTargetPlatform.Mac)
 		{
 			PublicDefinitions.Add("WITH_FACEBOOK=1");
diff --git a/Engine/Plugins/Online/OnlineSubsystemGoogle/Source/OnlineSubsystemGoogle.Build.cs b/Engine/Plugins/Online/OnlineSubsystemGoogle/Source/OnlineSubsystemGoogle.Build.cs
--- a/Engine/Plugins/Online/OnlineSubsystemGoogle/Source/OnlineSubsystemGoogle.Build.cs
+++ b/Engine/Plugins/Online/OnlineSubsystemGoogle/Source/OnlineSubsystemGoogle.Build.cs
@@ -72,14 +72,6 @@ public class OnlineSubsystemGoogle : ModuleRules
 		{
 			bUsesRestfulImpl = true;
 		}
-		else if (Target.Platform == UnrealTargetPlatform.XboxOne)
-		{
-			PrivateIncludePaths.Add("Private/XboxOne");
-		}
-		else if (Target.Platform == UnrealTargetPlatform.PS4)
-		{
-			PrivateIncludePaths.Add("Private/PS4");
-		}
 		else if (Target.Platform == UnrealTargetPlatform.Mac)
 		{
 			bUsesRestfulImpl = true;
