Removed unused Android Facebook SDK because actual dependency is downloaded using gradle (OnlineSubsystemFacebook_UPL.xml adds it)
Removed checks from that path from OnlineSubsystemFacebook.Build.cs

diff --git a/Engine/Plugins/Online/OnlineSubsystemFacebook/Source/OnlineSubsystemFacebook.Build.cs b/Engine/Plugins/Online/OnlineSubsystemFacebook/Source/OnlineSubsystemFacebook.Build.cs
--- a/Engine/Plugins/Online/OnlineSubsystemFacebook/Source/OnlineSubsystemFacebook.Build.cs
+++ b/Engine/Plugins/Online/OnlineSubsystemFacebook/Source/OnlineSubsystemFacebook.Build.cs
@@ -42,41 +42,17 @@ public class OnlineSubsystemFacebook : ModuleRules
 		}
 		else if (Target.Platform == UnrealTargetPlatform.Android)
 		{
-			bool bHasFacebookSDK = false;
-			string FacebookNFLDir = "";
-			try
-			{
-				FacebookNFLDir = System.IO.Path.Combine(EngineDirectory, "Restricted/NotForLicensees/Plugins/Online/OnlineSubsystemFacebook/Source/ThirdParty/Android/FacebookSDK");
-				bHasFacebookSDK = System.IO.Directory.Exists(FacebookNFLDir);
-			}
-			catch (System.Exception)
-			{
-			}
-
+			PublicDefinitions.Add("WITH_FACEBOOK=1");
 			PrivateIncludePaths.Add("Private/Android");
 
-			if (bHasFacebookSDK)
-			{
-				string Err = string.Format("Facebook SDK found in {0}", FacebookNFLDir);
-				System.Console.WriteLine(Err);
-
-				PublicDefinitions.Add("WITH_FACEBOOK=1");
-
-				PrivateDependencyModuleNames.AddRange(
+			PrivateDependencyModuleNames.AddRange(
 				new string[] {
 				"Launch",
 				}
-				);
+			);
 
-				string PluginPath = Utils.MakePathRelativeTo(ModuleDirectory, Target.RelativeEnginePath);
-				AdditionalPropertiesForReceipt.Add("AndroidPlugin", Path.Combine(PluginPath, "OnlineSubsystemFacebook_UPL.xml"));
-			}
-			else
-			{
-				string Err = string.Format("Facebook SDK not found in {0}", FacebookNFLDir);
-				System.Console.WriteLine(Err);
-				PublicDefinitions.Add("WITH_FACEBOOK=0");
-			}
+			string PluginPath = Utils.MakePathRelativeTo(ModuleDirectory, Target.RelativeEnginePath);
+			AdditionalPropertiesForReceipt.Add("AndroidPlugin", Path.Combine(PluginPath, "OnlineSubsystemFacebook_UPL.xml"));
 		}
 		else if (bUsesRestfulImpl)
 		{
