- Make FUniqueNetId::GetTypeHash virtual
- Remove some unused/unnecessary friend functions
- Convert some others to virtual overrides.

diff --git a/Engine/Plugins/Online/OnlineSubsystem/Source/Public/OnlineSubsystemTypes.h b/Engine/Plugins/Online/OnlineSubsystem/Source/Public/OnlineSubsystemTypes.h
--- a/Engine/Plugins/Online/OnlineSubsystem/Source/Public/OnlineSubsystemTypes.h
+++ b/Engine/Plugins/Online/OnlineSubsystem/Source/Public/OnlineSubsystemTypes.h
@@ -876,10 +876,9 @@ public:
 		}
 	}
 
-	/** Needed for TMap::GetTypeHash() */
-	friend uint32 GetTypeHash(const FUniqueNetIdString& A)
+	virtual uint32 GetTypeHash() const override
 	{
-		return ::GetTypeHash(A.UniqueNetIdStr);
+		return ::GetTypeHash(UniqueNetIdStr);
 	}
 
 public:
diff --git a/Engine/Plugins/Online/OnlineSubsystemEOS/Source/OnlineSubsystemEOS/Private/OnlineSubsystemEOSTypes.h b/Engine/Plugins/Online/OnlineSubsystemEOS/Source/OnlineSubsystemEOS/Private/OnlineSubsystemEOSTypes.h
--- a/Engine/Plugins/Online/OnlineSubsystemEOS/Source/OnlineSubsystemEOS/Private/OnlineSubsystemEOSTypes.h
+++ b/Engine/Plugins/Online/OnlineSubsystemEOS/Source/OnlineSubsystemEOS/Private/OnlineSubsystemEOSTypes.h
@@ -51,11 +51,6 @@ public:
 		return *static_cast<const FUniqueNetIdEOS*>(&NetId);
 	}
 
-	friend uint32 GetTypeHash(const FUniqueNetIdEOS& A)
-	{
-		return ::GetTypeHash(A.UniqueNetIdStr);
-	}
-
 	/** global static instance of invalid (zero) id */
 	static const FUniqueNetIdEOSRef& EmptyId()
 	{
diff --git a/Engine/Plugins/Online/OnlineSubsystemEOS/Source/OnlineSubsystemEOSPlus/Private/OnlineUserEOSPlus.h b/Engine/Plugins/Online/OnlineSubsystemEOS/Source/OnlineSubsystemEOSPlus/Private/OnlineUserEOSPlus.h
--- a/Engine/Plugins/Online/OnlineSubsystemEOS/Source/OnlineSubsystemEOSPlus/Private/OnlineUserEOSPlus.h
+++ b/Engine/Plugins/Online/OnlineSubsystemEOS/Source/OnlineSubsystemEOSPlus/Private/OnlineUserEOSPlus.h
@@ -114,20 +114,10 @@ public:
 	}
 	virtual FString ToDebugString() const override
 	{
+		// TODO OSS_UNIQUEID_REDACT
 		return ToString();
 	}
 
-	/** Needed for TMap::GetTypeHash() */
-	friend uint32 GetTypeHash(const FUniqueNetIdBinary& A)
-	{
-		uint32 Hash = 0;
-		for (uint8 CurrentByte : A.RawBytes)
-		{
-			Hash = HashCombine(Hash, GetTypeHash(CurrentByte));
-		}
-		return Hash;
-	}
-
 protected:
 	TArray<uint8> RawBytes;
 	FName Type;
diff --git a/Engine/Plugins/Online/OnlineSubsystemFacebook/Source/Private/OnlineSubsystemFacebookTypes.h b/Engine/Plugins/Online/OnlineSubsystemFacebook/Source/Private/OnlineSubsystemFacebookTypes.h
--- a/Engine/Plugins/Online/OnlineSubsystemFacebook/Source/Private/OnlineSubsystemFacebookTypes.h
+++ b/Engine/Plugins/Online/OnlineSubsystemFacebook/Source/Private/OnlineSubsystemFacebookTypes.h
@@ -64,6 +64,11 @@ public:
 		const FString UniqueNetIdStr = FString::Printf(TEXT("0%I64X"), UniqueNetId);
 		return OSS_UNIQUEID_REDACT(*this, UniqueNetIdStr);
 	}
+
+	virtual uint32 GetTypeHash() const override
+	{
+		return ::GetTypeHash(UniqueNetId);
+	}
 	//~ End FUniqueNetId Interface
 
 	/** global static instance of invalid (zero) id */
@@ -73,12 +78,6 @@ public:
 		return EmptyId;
 	}
 
-	/** Needed for TMap::GetTypeHash() */
-	friend uint32 GetTypeHash(const FUniqueNetIdFacebook& A)
-	{
-		return GetTypeHash(A.UniqueNetId);
-	}
-
 private:
 	/** Hidden on purpose */
 	FUniqueNetIdFacebook() :
diff --git a/Engine/Plugins/Online/OnlineSubsystemOculus/Source/Private/OnlineSubsystemOculusTypes.h b/Engine/Plugins/Online/OnlineSubsystemOculus/Source/Private/OnlineSubsystemOculusTypes.h
--- a/Engine/Plugins/Online/OnlineSubsystemOculus/Source/Private/OnlineSubsystemOculusTypes.h
+++ b/Engine/Plugins/Online/OnlineSubsystemOculus/Source/Private/OnlineSubsystemOculusTypes.h
@@ -87,10 +87,9 @@ public:
 		return TEXT("ovrID:") + OSS_UNIQUEID_REDACT(*this, UniqueNetIdStr);
 	}
 
-	/** Needed for TMap::GetTypeHash() */
-	friend uint32 GetTypeHash(const FUniqueNetIdOculus& A)
+	virtual uint32 GetTypeHash() const override
 	{
-		return GetTypeHash((uint64)A.ID);
+		return ::GetTypeHash((uint64)ID);
 	}
 
 	/** global static instance of invalid (zero) id */
diff --git a/Engine/Plugins/Online/OnlineSubsystemSteam/Source/Private/OnlineSubsystemSteamTypes.h b/Engine/Plugins/Online/OnlineSubsystemSteam/Source/Private/OnlineSubsystemSteamTypes.h
--- a/Engine/Plugins/Online/OnlineSubsystemSteam/Source/Private/OnlineSubsystemSteamTypes.h
+++ b/Engine/Plugins/Online/OnlineSubsystemSteam/Source/Private/OnlineSubsystemSteamTypes.h
@@ -168,10 +168,9 @@ public:
 		return FUniqueNetIdSteam(SteamID).ToDebugString();
 	}
 
-	/** Needed for TMap::GetTypeHash() */
-	friend uint32 GetTypeHash(const FUniqueNetIdSteam& A)
+	virtual uint32 GetTypeHash() const override
 	{
-		return GetTypeHash(A.UniqueNetId);
+		return ::GetTypeHash(UniqueNetId);
 	}
 
 	/** global static instance of invalid (zero) id */
diff --git a/Engine/Plugins/Online/OnlineSubsystemTencent/Source/Public/OnlineSubsystemTencentTypes.h b/Engine/Plugins/Online/OnlineSubsystemTencent/Source/Public/OnlineSubsystemTencentTypes.h
--- a/Engine/Plugins/Online/OnlineSubsystemTencent/Source/Public/OnlineSubsystemTencentTypes.h
+++ b/Engine/Plugins/Online/OnlineSubsystemTencent/Source/Public/OnlineSubsystemTencentTypes.h
@@ -122,11 +122,10 @@ public:
 		}
 	}
 
-	/** Needed for TMap::GetTypeHash() */
-	friend uint32 GetTypeHash(const FUniqueNetIdRail& A)
+	virtual uint32 GetTypeHash() const override
 	{
 		uint64 id = A.RailID.get_id();
-		return GetTypeHash(id);
+		return ::GetTypeHash(id);
 	}
 
 	/** Convenience cast to RailID */
diff --git a/Engine/Source/Runtime/CoreOnline/Public/Online/CoreOnline.h b/Engine/Source/Runtime/CoreOnline/Public/Online/CoreOnline.h
--- a/Engine/Source/Runtime/CoreOnline/Public/Online/CoreOnline.h
+++ b/Engine/Source/Runtime/CoreOnline/Public/Online/CoreOnline.h
@@ -184,6 +184,11 @@ public:
 	 */
 	virtual FString ToDebugString() const = 0;
 
+	virtual uint32 GetTypeHash() const
+	{
+		return CityHash32(reinterpret_cast<const char*>(GetBytes()), GetSize());
+	}
+
 	/**
 	 * @return hex encoded string representation of unique id
 	 */
@@ -198,7 +203,7 @@ public:
 
 	friend inline uint32 GetTypeHash(const FUniqueNetId& Value)
 	{
-		return CityHash32(reinterpret_cast<const char*>(Value.GetBytes()), Value.GetSize());
+		return Value.GetTypeHash();
 	}
 };
 
