Fix some deprecations and issues found in OnlineSubsystemFacebook RESTful implementation
 - Using bUsePopup=false was making very easy to break the login flow
 - Added some commands to test facebook login/logout/list friends
 - Removed deprecated permissions and fields
 - Fixed redirection parsing. It was not properly handling errors so login was kept unfinished
 - Fixed parsing user picture data

diff --git a/Engine/Config/BaseEngine.ini b/Engine/Config/BaseEngine.ini
--- a/Engine/Config/BaseEngine.ini
+++ b/Engine/Config/BaseEngine.ini
@@ -2041,26 +2041,20 @@ RedirectPort=9001
 
 [OnlineSubsystemFacebook]
 bEnabled=false
-APIVer="v2.12"
+APIVer="v15.0"
 
 [OnlineSubsystemFacebook.OnlineIdentityFacebook]
 LoginUrl="https://www.facebook.com/`ver/dialog/oauth"
 LoginRedirectUrl="https://www.facebook.com/connect/login_success.html"
 MeURL="https://graph.facebook.com/`ver/me?access_token=`token"
 +LoginDomains=".facebook.com"
-bUsePopup=false
-ProfileFields=locale
-+ProfileFields=link
-+ProfileFields=gender
+bUsePopup=true
 
 [OnlineSubsystemFacebook.OnlineSharingFacebook]
 PermissionsURL="https://graph.facebook.com/`ver/me/permissions?access_token=`token"
 
 [OnlineSubsystemFacebook.OnlineFriendsFacebook]
 FriendsUrl="https://graph.facebook.com/`ver/me/friends?fields=`fields&access_token=`token"
-FriendsFields=locale
-+FriendsFields=link
-+FriendsFields=gender
 
 [OnlineSubsystemSamsung]
 bEnabled=true
diff --git a/Engine/Plugins/Online/OnlineSubsystemFacebook/Source/Private/OnlineFriendsFacebookCommon.h b/Engine/Plugins/Online/OnlineSubsystemFacebook/Source/Private/OnlineFriendsFacebookCommon.h
--- a/Engine/Plugins/Online/OnlineSubsystemFacebook/Source/Private/OnlineFriendsFacebookCommon.h
+++ b/Engine/Plugins/Online/OnlineSubsystemFacebook/Source/Private/OnlineFriendsFacebookCommon.h
@@ -99,7 +99,7 @@ private:
 	// FJsonSerializable
 	BEGIN_ONLINE_JSON_SERIALIZER
 		ONLINE_JSON_SERIALIZE(FRIEND_FIELD_ID, UserIdStr);
-		ONLINE_JSON_SERIALIZE_OBJECT_SERIALIZABLE(PICTURE_DATA, Picture);
+		ONLINE_JSON_SERIALIZE_OBJECT_SERIALIZABLE(FRIEND_FIELD_PICTURE, Picture);
 		ONLINE_JSON_SERIALIZE_SIMPLE_COPY(AccountData);
 	END_ONLINE_JSON_SERIALIZER
 
diff --git a/Engine/Plugins/Online/OnlineSubsystemFacebook/Source/Private/OnlineSharingFacebookCommon.cpp b/Engine/Plugins/Online/OnlineSubsystemFacebook/Source/Private/OnlineSharingFacebookCommon.cpp
--- a/Engine/Plugins/Online/OnlineSubsystemFacebook/Source/Private/OnlineSharingFacebookCommon.cpp
+++ b/Engine/Plugins/Online/OnlineSubsystemFacebook/Source/Private/OnlineSharingFacebookCommon.cpp
@@ -18,20 +18,9 @@ void FFacebookPermissions::Setup()
 {
 	///////////////////////////////////////////////
 	// Read Permissions
-	SharingPermissionsMap.Add(EOnlineSharingCategory::ReadPosts, {TEXT(PERM_READ_STREAM)});
-	SharingPermissionsMap.Add(EOnlineSharingCategory::Friends, {TEXT(PERM_READ_FRIENDS)});
-	SharingPermissionsMap.Add(EOnlineSharingCategory::Email, {TEXT(PERM_READ_EMAIL)});
-	SharingPermissionsMap.Add(EOnlineSharingCategory::Mailbox, {TEXT(PERM_READ_MAILBOX)});
-	SharingPermissionsMap.Add(EOnlineSharingCategory::OnlineStatus, {TEXT(PERM_READ_STATUS), TEXT(PERM_READ_PRESENCE)});
-	SharingPermissionsMap.Add(EOnlineSharingCategory::ProfileInfo, {TEXT(PERM_PUBLIC_PROFILE)});
-	SharingPermissionsMap.Add(EOnlineSharingCategory::LocationInfo, {TEXT(PERM_READ_CHECKINS), TEXT(PERM_READ_HOMETOWN)});
-
-	///////////////////////////////////////////////
-	// Publish Permissions
-	SharingPermissionsMap.Add(EOnlineSharingCategory::SubmitPosts, {TEXT(PERM_PUBLISH_ACTION)});
-	SharingPermissionsMap.Add(EOnlineSharingCategory::ManageFriends, {TEXT(PERM_MANAGE_FRIENDSLIST)});
-	SharingPermissionsMap.Add(EOnlineSharingCategory::AccountAdmin, {TEXT(PERM_MANAGE_NOTIFICATIONS)});
-	SharingPermissionsMap.Add(EOnlineSharingCategory::Events, {TEXT(PERM_CREATE_EVENT), TEXT(PERM_RSVP_EVENT)});
+	SharingPermissionsMap.Add(EOnlineSharingCategory::Friends, { TEXT(PERM_READ_FRIENDS) });
+	SharingPermissionsMap.Add(EOnlineSharingCategory::Email, { TEXT(PERM_READ_EMAIL) });
+	SharingPermissionsMap.Add(EOnlineSharingCategory::ProfileInfo, { TEXT(PERM_PUBLIC_PROFILE), TEXT(PERM_PUBLIC_GAMING_PROFILE) });
 }
 
 void FFacebookPermissions::Reset()
diff --git a/Engine/Plugins/Online/OnlineSubsystemFacebook/Source/Private/OnlineSharingFacebookCommon.h b/Engine/Plugins/Online/OnlineSubsystemFacebook/Source/Private/OnlineSharingFacebookCommon.h
--- a/Engine/Plugins/Online/OnlineSubsystemFacebook/Source/Private/OnlineSharingFacebookCommon.h
+++ b/Engine/Plugins/Online/OnlineSubsystemFacebook/Source/Private/OnlineSharingFacebookCommon.h
@@ -14,23 +14,11 @@ class FOnlineSubsystemFacebook;
 
 /** Permissions to request at login */
 #define PERM_PUBLIC_PROFILE "public_profile"
+#define PERM_PUBLIC_GAMING_PROFILE "gaming_profile" // Same as public_profile for facebook apps with Gaming Services type
+#define PERM_READ_EMAIL "email"
 
 /** Elevated read permissions */
 #define PERM_READ_FRIENDS "user_friends"
-#define PERM_READ_EMAIL "email"
-#define PERM_READ_STREAM "read_stream"
-#define PERM_READ_MAILBOX "read_mailbox"
-#define PERM_READ_STATUS "user_status"
-#define PERM_READ_PRESENCE "user_online_presence"
-#define PERM_READ_CHECKINS "user_checkins"
-#define PERM_READ_HOMETOWN "user_hometown"
-
-/** Elevated publish permissions */
-#define PERM_PUBLISH_ACTION "publish_actions"
-#define PERM_MANAGE_FRIENDSLIST "manage_friendlists"
-#define PERM_MANAGE_NOTIFICATIONS "manage_notifications"
-#define PERM_CREATE_EVENT "create_event"
-#define PERM_RSVP_EVENT "rsvp_event"
 
 /** Facebook Permission Json */
 #define PERM_JSON_PERMISSIONS "data"
diff --git a/Engine/Plugins/Online/OnlineSubsystemFacebook/Source/Private/OnlineSubsystemFacebookCommon.cpp b/Engine/Plugins/Online/OnlineSubsystemFacebook/Source/Private/OnlineSubsystemFacebookCommon.cpp
--- a/Engine/Plugins/Online/OnlineSubsystemFacebook/Source/Private/OnlineSubsystemFacebookCommon.cpp
+++ b/Engine/Plugins/Online/OnlineSubsystemFacebook/Source/Private/OnlineSubsystemFacebookCommon.cpp
@@ -87,11 +87,16 @@ FString FOnlineSubsystemFacebookCommon::GetAppId() const
 
 bool FOnlineSubsystemFacebookCommon::Exec(UWorld* InWorld, const TCHAR* Cmd, FOutputDevice& Ar)
 {
-	if (FOnlineSubsystemImpl::Exec(InWorld, Cmd, Ar))
+	bool bWasHandled = false;
+	if (FParse::Command(&Cmd, TEXT("FACEBOOK")))
 	{
-		return true;
+		bWasHandled = HandleFacebookExecCommands(InWorld, Cmd, Ar);
 	}
-	return false;
+	else if (FOnlineSubsystemImpl::Exec(InWorld, Cmd, Ar))
+	{
+		bWasHandled = true;
+	}
+	return bWasHandled;
 }
 
 IOnlineSessionPtr FOnlineSubsystemFacebookCommon::GetSessionInterface() const
@@ -213,3 +218,127 @@ FText FOnlineSubsystemFacebookCommon::GetOnlineServiceName() const
 {
 	return NSLOCTEXT("OnlineSubsystemFacebook", "OnlineServiceName", "Facebook");
 }
+
+bool FOnlineSubsystemFacebookCommon::HandleFacebookExecCommands(UWorld* InWorld, const TCHAR* Cmd, FOutputDevice& Ar)
+{
+	bool bWasHandled = false;
+
+	if (FParse::Command(&Cmd, TEXT("LOGIN")))
+	{
+		bWasHandled = true;
+
+		FString LocalNumStr = FParse::Token(Cmd, false);
+		int32 LocalNum = FCString::Atoi(*LocalNumStr);
+		if (LocalNumStr.IsEmpty() || LocalNum < 0 || LocalNum > MAX_LOCAL_PLAYERS)
+		{
+			UE_LOG_ONLINE(Warning, TEXT("usage: LOGIN <localnum>"));
+		}
+		else
+		{
+			HandleFacebookLoginCommand(LocalNum);
+		}
+	}
+	else if (FParse::Command(&Cmd, TEXT("LOGOUT")))
+	{
+		bWasHandled = true;
+
+		FString LocalNumStr = FParse::Token(Cmd, false);
+		int32 LocalNum = FCString::Atoi(*LocalNumStr);
+		if (LocalNumStr.IsEmpty() || LocalNum < 0 || LocalNum > MAX_LOCAL_PLAYERS)
+		{
+			UE_LOG_ONLINE(Warning, TEXT("usage: LOGOUT <localnum>"));
+		}
+		else
+		{
+			HandleFacebookLogoutCommand(LocalNum);
+		}
+	}
+	else if (FParse::Command(&Cmd, TEXT("FRIENDS")))
+	{
+		bWasHandled = true;
+
+		FString LocalNumStr = FParse::Token(Cmd, false);
+		int32 LocalNum = FCString::Atoi(*LocalNumStr);
+		if (LocalNumStr.IsEmpty() || LocalNum < 0 || LocalNum > MAX_LOCAL_PLAYERS)
+		{
+			UE_LOG_ONLINE(Warning, TEXT("usage: LOGOUT <localnum>"));
+		}
+		else
+		{
+			HandleFacebookFriendsCommand(LocalNum);
+		}
+	}
+	return bWasHandled;
+}
+
+void FOnlineSubsystemFacebookCommon::HandleFacebookLoginCommand(int32 LocalNum)
+{
+	IOnlineIdentityPtr IdentityInt = GetIdentityInterface();
+	if (IdentityInt.IsValid())
+	{
+		if (IdentityInt->GetLoginStatus(LocalNum) == ELoginStatus::LoggedIn)
+		{
+			UE_LOG_ONLINE(Warning, TEXT("User %d already logged into Facebook. Nickname %s"), LocalNum, *IdentityInt->GetPlayerNickname(LocalNum));
+		}
+		else
+		{
+			FOnlineAccountCredentials AccountCredentials("", "", IdentityInt->GetAuthType());
+
+			TSharedPtr<FDelegateHandle> LoginCompleteHandle = MakeShared<FDelegateHandle>();
+			*LoginCompleteHandle = IdentityInt->AddOnLoginCompleteDelegate_Handle(LocalNum, FOnLoginCompleteDelegate::CreateLambda([IdentityInt = FacebookIdentity.Get(), LoginCompleteHandle](int32 LocalUserNum, bool bWasSuccessful, const FUniqueNetId& /*UserId*/, const FString& Error)
+				{
+					if (bWasSuccessful)
+					{
+						UE_LOG_ONLINE(Log, TEXT("Facebook login succeeded for player %d. User nickname %s"), LocalUserNum, *IdentityInt->GetPlayerNickname(LocalUserNum));
+					}
+					else
+					{
+						UE_LOG_ONLINE(Log, TEXT("Facebook login failed for player %d: Error: %s"), LocalUserNum, *Error);
+					}
+					IdentityInt->ClearOnLoginCompleteDelegate_Handle(LocalUserNum, *LoginCompleteHandle);
+				}));
+			IdentityInt->Login(LocalNum, AccountCredentials);
+		}
+	}
+}
+
+void FOnlineSubsystemFacebookCommon::HandleFacebookLogoutCommand(int32 LocalNum)
+{
+	IOnlineIdentityPtr IdentityInt = GetIdentityInterface();
+	if (IdentityInt.IsValid())
+	{
+		if (IdentityInt->GetLoginStatus(LocalNum) == ELoginStatus::NotLoggedIn)
+		{
+			UE_LOG_ONLINE(Warning, TEXT("User %d not logged into Facebook"), LocalNum);
+			return;
+		}
+		TSharedPtr<FDelegateHandle> LogoutCompleteHandle = MakeShared<FDelegateHandle>();
+		*LogoutCompleteHandle = IdentityInt->AddOnLogoutCompleteDelegate_Handle(LocalNum, FOnLogoutCompleteDelegate::CreateLambda([IdentityInt = IdentityInt.Get(), LogoutCompleteHandle](int32 LocalUserNum, bool bWasSuccessful)
+			{
+				UE_LOG_ONLINE(Log, TEXT("Facebook logout %s for player %d"), bWasSuccessful ? TEXT("succeeded") : TEXT("failed"), LocalUserNum);
+				IdentityInt->ClearOnLogoutCompleteDelegate_Handle(LocalUserNum, *LogoutCompleteHandle);
+			}));
+		IdentityInt->Logout(LocalNum);
+	}
+}
+
+void FOnlineSubsystemFacebookCommon::HandleFacebookFriendsCommand(int32 LocalNum)
+{
+	IOnlineFriendsPtr FriendsInt = GetFriendsInterface();
+	if (FriendsInt.IsValid())
+	{
+		FriendsInt->ReadFriendsList(LocalNum, ToString(EFriendsLists::Default), FOnReadFriendsListComplete::CreateLambda([FriendsInt = FriendsInt.Get()](int32 LocalUserNum, bool bWasSuccessful, const FString& ListName, const FString& ErrorStr) {
+			if (bWasSuccessful)
+			{
+				UE_LOG_ONLINE(Log, TEXT("Retrieve Facebook friends list %s for player %d succeeded"), *ListName, LocalUserNum);
+				TArray<TSharedRef<FOnlineFriend>> Friends;
+				FriendsInt->GetFriendsList(LocalUserNum, ListName, Friends);
+				UE_LOG_ONLINE(Log, TEXT("List of %d friends: %s"), Friends.Num(), *FString::JoinBy(Friends, TEXT(", "), [](const TSharedRef<FOnlineFriend>& Friend) { return Friend->GetRealName(); }));
+			}
+			else
+			{
+				UE_LOG_ONLINE(Warning, TEXT("Retrieve Facebook friends list %s for player %d failed: Error: %s"), *ListName, LocalUserNum, *ErrorStr);
+			}
+			}));
+	}
+}
diff --git a/Engine/Plugins/Online/OnlineSubsystemFacebook/Source/Private/OnlineSubsystemFacebookCommon.h b/Engine/Plugins/Online/OnlineSubsystemFacebook/Source/Private/OnlineSubsystemFacebookCommon.h
--- a/Engine/Plugins/Online/OnlineSubsystemFacebook/Source/Private/OnlineSubsystemFacebookCommon.h
+++ b/Engine/Plugins/Online/OnlineSubsystemFacebook/Source/Private/OnlineSubsystemFacebookCommon.h
@@ -20,6 +20,19 @@ typedef TSharedPtr<class FOnlineExternalUIFacebookCommon, ESPMode::ThreadSafe> F
 class ONLINESUBSYSTEMFACEBOOK_API FOnlineSubsystemFacebookCommon 
 	: public FOnlineSubsystemImpl
 {
+private:
+
+	/**
+	 * Exec function handling for Exec() call
+	 */
+	bool HandleFacebookExecCommands(UWorld* InWorld, const TCHAR* Cmd, FOutputDevice& Ar);
+
+	/** Handle facebook login console command */
+	void HandleFacebookLoginCommand(int32 LocalNum);
+	/** Handle facebook logout console command */
+	void HandleFacebookLogoutCommand(int32 LocalNum);
+	/** Handle facebook friends console command */
+	void HandleFacebookFriendsCommand(int32 LocalNum);
 public:
 
 	// IOnlineSubsystem Interface
diff --git a/Engine/Plugins/Online/OnlineSubsystemFacebook/Source/Private/Rest/OnlineExternalUIInterfaceFacebookRest.cpp b/Engine/Plugins/Online/OnlineSubsystemFacebook/Source/Private/Rest/OnlineExternalUIInterfaceFacebookRest.cpp
--- a/Engine/Plugins/Online/OnlineSubsystemFacebook/Source/Private/Rest/OnlineExternalUIInterfaceFacebookRest.cpp
+++ b/Engine/Plugins/Online/OnlineSubsystemFacebook/Source/Private/Rest/OnlineExternalUIInterfaceFacebookRest.cpp
@@ -68,23 +68,53 @@ FLoginFlowResult FOnlineExternalUIFacebook::ParseRedirectResult(const FFacebookL
 
 	TMap<FString, FString> ParamsMap;
 	{
-		FString URLPrefix;
-		FString QueryParams;
-		if (!RedirectURL.Split(TEXT("#"), &URLPrefix, &QueryParams))
+		FStringView QueryParams;
+		FStringView UriFragment;
+
+		int32 QueryParamsStart = RedirectURL.Find(TEXT("?"));
+		int32 FragmentStart = RedirectURL.Find(TEXT("#"), ESearchCase::IgnoreCase, ESearchDir::FromStart, QueryParamsStart);
+		
+		bool HasQueryParams = QueryParamsStart != INDEX_NONE;
+		bool HasUriFragment = FragmentStart != INDEX_NONE;
+
+		if (HasQueryParams)
 		{
-			QueryParams = RedirectURL;
+			// Discard "?"
+			QueryParamsStart += 1;
+			int32 QueryParamsLength = HasUriFragment ? (FragmentStart - QueryParamsStart) : (RedirectURL.Len() - 1);
+			QueryParams = FStringView(RedirectURL).Mid(QueryParamsStart, QueryParamsLength);
+			HasQueryParams = QueryParams.IsEmpty();
 		}
 
-		// Remove the "Facebook fragment"
-		// https://developers.facebook.com/blog/post/552/
-		FString ParamsOnly;
-		if (!QueryParams.Split(TEXT("#_=_"), &ParamsOnly, nullptr))
+		if (HasUriFragment)
 		{
-			ParamsOnly = QueryParams;
+			UriFragment = FStringView(RedirectURL).RightChop(FragmentStart);
+
+			// Facebook adds a fragment #_=_ when that field is left blank
+			// https://developers.facebook.com/blog/post/552/
+			if (UriFragment.Equals(TEXT("#_=_")))
+			{
+				UriFragment.Reset();
+				HasUriFragment = false;
+			}
+			else
+			{
+				UriFragment.RemovePrefix(1);
+			}
 		}
 
+		// Auth token response is received through the UriFragment
+		// Error notifications are received through the QueryParams
 		TArray<FString> Params;
-		ParamsOnly.ParseIntoArray(Params, TEXT("&"));
+		if (HasUriFragment)
+		{
+			FString(UriFragment).ParseIntoArray(Params, TEXT("&"));
+		}
+		else
+		{
+			FString(QueryParams).ParseIntoArray(Params, TEXT("&"));
+		}
+
 		for (FString& Param : Params)
 		{
 			FString Key, Value;
@@ -95,6 +125,8 @@ FLoginFlowResult FOnlineExternalUIFacebook::ParseRedirectResult(const FFacebookL
 		}
 	}
 
+	// Both error or success should contain the state token
+	// Otherwise the response does not refer to this request
 	const FString* State = ParamsMap.Find(FB_STATE_TOKEN);
 	if (State)
 	{
