Fix missing include reported by CI

https://horde.devtools.epicgames.com/job/6373761f3650332de8bd4745?step=48af&issue=236055

diff --git a/Engine/Plugins/Online/OnlineSubsystemFacebook/Source/Private/Android/OnlineSubsystemFacebook.cpp b/Engine/Plugins/Online/OnlineSubsystemFacebook/Source/Private/Android/OnlineSubsystemFacebook.cpp
--- a/Engine/Plugins/Online/OnlineSubsystemFacebook/Source/Private/Android/OnlineSubsystemFacebook.cpp
+++ b/Engine/Plugins/Online/OnlineSubsystemFacebook/Source/Private/Android/OnlineSubsystemFacebook.cpp
@@ -5,6 +5,7 @@
 #include "OnlineIdentityFacebook.h"
 #include "OnlineFriendsFacebook.h"
 #include "OnlineSharingFacebook.h"
+#include "OnlineUserFacebook.h"
 
 FOnlineSubsystemFacebook::FOnlineSubsystemFacebook(FName InInstanceName)
 	: FOnlineSubsystemFacebookCommon(InInstanceName)
