Set Java 1.8 for sourceCompatibility and targetCompatibility

diff --git a/Engine/Plugins/Online/OnlineSubsystemFacebook/Source/OnlineSubsystemFacebookAndroid_UPL.xml b/Engine/Plugins/Online/OnlineSubsystemFacebook/Source/OnlineSubsystemFacebookAndroid_UPL.xml
--- a/Engine/Plugins/Online/OnlineSubsystemFacebook/Source/OnlineSubsystemFacebookAndroid_UPL.xml
+++ b/Engine/Plugins/Online/OnlineSubsystemFacebook/Source/OnlineSubsystemFacebookAndroid_UPL.xml
@@ -40,13 +40,6 @@
 		<if condition="bEnabled">
 			<true>
 				<insert>
-
-android {
-	compileOptions {
-        sourceCompatibility JavaVersion.VERSION_1_8
-        targetCompatibility JavaVersion.VERSION_1_8
-    }	
-}
 dependencies {
 	implementation('com.facebook.android:facebook-android-sdk:15.0.1')
 	implementation('com.facebook.android:facebook-login:15.0.1')
