FiX OSSFacebook UPL smashing CFBundleURLTypes in Info.plist instead of adding to it

diff --git a/Engine/Plugins/Online/OnlineSubsystemFacebook/Source/OnlineSubsystemFacebookIOS_UPL.xml b/Engine/Plugins/Online/OnlineSubsystemFacebook/Source/OnlineSubsystemFacebookIOS_UPL.xml
--- a/Engine/Plugins/Online/OnlineSubsystemFacebook/Source/OnlineSubsystemFacebookIOS_UPL.xml
+++ b/Engine/Plugins/Online/OnlineSubsystemFacebook/Source/OnlineSubsystemFacebookIOS_UPL.xml
@@ -9,11 +9,12 @@
 		<setBoolFromProperty result="bEnabled" ini="Engine" section="OnlineSubsystemFacebook" property="bEnabled" default="false"/>
 		<setStringFromProperty result="ClientId" ini="Engine" section="OnlineSubsystemFacebook" property="ClientId" default="0"/>
 		<setStringFromProperty result="ClientToken" ini="Engine" section="OnlineSubsystemFacebook" property="ClientToken" default="0"/>
-		<setStringFromProperty result="BundleDisplayName" ini="Engine" section="/Script/IOSRuntimeSettings.IOSRuntimeSettings" property="BundleDisplayName" defaule=""/>
+		<setStringFromProperty result="BundleDisplayName" ini="Engine" section="/Script/IOSRuntimeSettings.IOSRuntimeSettings" property="BundleDisplayName" default=""/>
 		<setBoolFromProperty result="bEnableAutomaticLogging" ini="Engine" section="OnlineSubsystemFacebook" property="bEnableAutomaticLogging" default="1"/>
 		<setBoolFromProperty result="bEnableAdvertisingId" ini="Engine" section="OnlineSubsystemFacebook" property="bEnableAdvertisingId" default="1"/>
 		<dumpvars/>
 	</init>
+
 	<iosPListUpdates>
 		<if condition="bEnabled">
 			<true>
@@ -45,22 +46,79 @@
 				</if>
 
 				<!-- CFBundleURLTypes block -->
-				<addElements tag="dict" once="true"> 
-					<key>CFBundleURLTypes</key>
-				</addElements>
+				<!-- Insert our schemas into the first CFBundleURLSchemes array we find inside CFBundleURLTypes content-->
+				<setBool result="bProcessingComplete" value="false"/>
+				<loopElements tag="dict">
+					<setBool result="bCFBundleURLTypesFound" value="false"/>
+
+					<loopElements tag="$">
+						<setStringFromTag result="TagName" tag="$"/>
+
+						<!-- If we found CFBundleURLTypes we are on next iteration and be should be on its <array > value -->
+						<if condition="bCFBundleURLTypesFound">
+							<true>
+								<setBoolIsEqual result="bIsArray" arg1="$S(TagName)" arg2="array"/>
+								<if condition="bIsArray">
+									<true>
+										<!-- this element is an array of dictionaries -->
+										<loopElements tag="$">
+											<setBool result="bCFBundleURLSchemesFound" value="false"/>
+											<!-- iterate on the <dict> entries looking for a <key> with the text CFBundleURLSchemes to insert our schemas-->
+											<loopElements tag="$">
+												<setStringFromTag result="TagName" tag="$"/>
+
+												<!-- If we found CFBundleURLSchemes we are on next iteration and be should be on its <array > value, so we can add our schemas -->
+												<if condition="bCFBundleURLSchemesFound">
+													<true>
+														<setBoolIsEqual result="bIsArray" arg1="$S(TagName)" arg2="array"/>
+														<if condition="bIsArray">
+															<true>
+																<!-- Add schema to current array -->
+																<setElement result="Schema" value="string" text="fb$S(ClientId)"/>
+																<addElement tag="$" name="Schema"/>
+
+																<setBool result="bProcessingComplete" value="true"/>
+															</true>
+														</if>
+													</true>
+												</if>
+
+												<if condition="bProcessingComplete">
+													<false>
+														<!-- Check if we found a <key> with text CFBundleURLSchemes -->
+														<setBoolIsEqual result="bIsKey" arg1="$S(TagName)" arg2="key"/>
+														<if condition="bIsKey">
+															<true>
+																<setStringFromTagText result="TagValue" tag="$"/>
+																<setBoolIsEqual result="bCFBundleURLSchemesFound" arg1="$S(TagValue)" arg2="CFBundleURLSchemes"/>
+															</true>
+														</if>
+													</false>
+												</if>
+											</loopElements>
+										</loopElements>
+									</true>
+								</if>
+							</true>
+						</if>
+
+						<if condition="bProcessingComplete">
+							<false>
+								<!-- Check if we found a <key> with text CFBundleURLTypes -->
+								<setBoolIsEqual result="bIsKey" arg1="$S(TagName)" arg2="key"/>
+								<if condition="bIsKey">
+									<true>
+										<setStringFromTagText result="TagValue" tag="$"/>
+										<setBoolIsEqual result="bCFBundleURLTypesFound" arg1="$S(TagValue)" arg2="CFBundleURLTypes"/>
+									</true>
+								</if>
+							</false>
+						</if>
+					</loopElements>
+				</loopElements>
+
 
-				<setElement result="FbIdString" value="string" text="fb$S(ClientId)"/>
-				<setElement result="ValueForKey" value="array"/>
-				<addElement tag="$ValueForKey" name="FbIdString"/>
 
-				<setElement result="DictOnArray" value="dict"/>
-				<addElements tag="$DictOnArray">
-  					<key>CFBundleURLSchemes</key>
-				</addElements>
-				<addElement tag="$DictOnArray" name="ValueForKey"/>
-				<setElement result="NewEntry" value="array"/>				
-				<addElement tag="$NewEntry" name="DictOnArray"/>
-				<addElement tag="dict" name="NewEntry" once="true"/> 
 				<!-- End of CFBundleURLTypes block -->
 
 				<if condition="bEnableAdvertisingId">
