Add missing file for Facebook SDK update on iOS

diff --git a/Engine/Plugins/Online/OnlineSubsystemFacebook/Source/Private/IOS/OnlineIdentityFacebook.cpp b/Engine/Plugins/Online/OnlineSubsystemFacebook/Source/Private/IOS/OnlineIdentityFacebook.cpp
--- a/Engine/Plugins/Online/OnlineSubsystemFacebook/Source/Private/IOS/OnlineIdentityFacebook.cpp
+++ b/Engine/Plugins/Online/OnlineSubsystemFacebook/Source/Private/IOS/OnlineIdentityFacebook.cpp
@@ -11,18 +11,62 @@
 #include "IOS/IOSAsyncTask.h"
 
 THIRD_PARTY_INCLUDES_START
+#if UE_WITH_CLASSIC_FACEBOOK_LOGIN
+#import <AppTrackingTransparency/AppTrackingTransparency.h>
+#endif
+#import <AuthenticationServices/AuthenticationServices.h>
+#import <SafariServices/SafariServices.h>
 #import <FBSDKCoreKit/FBSDKCoreKit.h>
+#import <FBSDKCoreKit/FBSDKCoreKit-Swift.h>
 #import <FBSDKLoginKit/FBSDKLoginKit-Swift.h>
 THIRD_PARTY_INCLUDES_END
 
 ///////////////////////////////////////////////////////////////////////////////////////
 // FOnlineIdentityFacebook implementation
 
+static TOptional<bool> ShouldUseClassicLogin()
+{
+	bool bUseClassicLogin = true;
+	bool bFallbackToLimitedLogin = false;
+
+	GConfig->GetBool(TEXT("OnlineSubsystemFacebook"), TEXT("bUseClassicLogin"), bUseClassicLogin, GEngineIni);
+	GConfig->GetBool(TEXT("OnlineSubsystemFacebook"), TEXT("bFallbackToLimitedLogin"), bFallbackToLimitedLogin, GEngineIni);
+
+	bool bClassicLoginAllowed = false;
+#if UE_WITH_CLASSIC_FACEBOOK_LOGIN
+	switch(ATTrackingManager.trackingAuthorizationStatus)
+	{
+		case ATTrackingManagerAuthorizationStatusAuthorized:
+			bClassicLoginAllowed = true;
+			break;
+		case ATTrackingManagerAuthorizationStatusDenied:
+		case ATTrackingManagerAuthorizationStatusNotDetermined:
+		case ATTrackingManagerAuthorizationStatusRestricted:
+			bClassicLoginAllowed = false;
+			break;
+	}
+#endif
+
+	if (bUseClassicLogin && !bClassicLoginAllowed)
+	{
+		if (bFallbackToLimitedLogin)
+		{
+			bUseClassicLogin = false;
+			UE_LOG_ONLINE_IDENTITY(Warning, TEXT("Falling back to Limited Facebook login because application tracking was not authorized"));
+		}
+		else
+		{
+			UE_LOG_ONLINE_IDENTITY(Error, TEXT("Classic Facebook login is not supported if application tracking was not authorized"));
+			return NullOpt;
+		}
+	}
+	return bUseClassicLogin;
+}
+
 FOnlineIdentityFacebook::FOnlineIdentityFacebook(FOnlineSubsystemFacebook* InSubsystem)
 	: FOnlineIdentityFacebookCommon(InSubsystem)
-	, LoginStatus(ELoginStatus::NotLoggedIn)
 {
-	// Setup permission scope fields
+	// Setup scopes fields
 	GConfig->GetArray(TEXT("OnlineSubsystemFacebook.OnlineIdentityFacebook"), TEXT("ScopeFields"), ScopeFields, GEngineIni);
 	// always required fields
 	ScopeFields.AddUnique(TEXT(PERM_PUBLIC_PROFILE));
@@ -55,26 +99,43 @@ void FOnlineIdentityFacebook::OnFacebookProfileChange(FBSDKProfile* OldProfile,
 	UE_LOG_ONLINE_IDENTITY(Warning, TEXT("FOnlineIdentityFacebook::OnFacebookProfileChange Old: %p New: %p"), OldProfile, NewProfile);
 }
 
+bool FOnlineIdentityFacebook::IsUsingClassicLogin() const 
+{
+	return bIsUsingClassicLogin;
+}
+
 bool FOnlineIdentityFacebook::Login(int32 LocalUserNum, const FOnlineAccountCredentials& AccountCredentials)
 {
 	bool bTriggeredLogin = true;
 
+	if (bIsLoginInProgress)
+	{
+		TriggerOnLoginCompleteDelegates(LocalUserNum, false, *FUniqueNetIdFacebook::EmptyId(), TEXT("Login already in progress"));
+		return false;
+	}
+
 	if (GetLoginStatus(LocalUserNum) != ELoginStatus::NotLoggedIn)
 	{
 		TriggerOnLoginCompleteDelegates(LocalUserNum, true, *GetUniquePlayerId(LocalUserNum), TEXT("Already logged in"));
 		return false;
 	}
 
-	ensure(LoginStatus == ELoginStatus::NotLoggedIn);
+	TOptional<bool> UseClassicLogin = ShouldUseClassicLogin();
+	if (!UseClassicLogin.IsSet())
+	{
+		TriggerOnLoginCompleteDelegates(LocalUserNum, false, *FUniqueNetIdFacebook::EmptyId(), TEXT("Login type unsupported"));
+		return false;
+	}
 
+	bIsLoginInProgress = true;
+	bIsUsingClassicLogin = *UseClassicLogin;
+	
 	dispatch_async(dispatch_get_main_queue(),^
 		{
-			FBSDKAccessToken *accessToken = [FBSDKAccessToken currentAccessToken];
-			if (accessToken == nil)
+			if ((bIsUsingClassicLogin && (FBSDKAccessToken.currentAccessToken == nil || [FBSDKAccessToken.currentAccessToken isExpired])) ||
+				(!bIsUsingClassicLogin && FBSDKProfile.currentProfile == nil) )
 			{
 				FBSDKLoginManager* loginManager = [[FBSDKLoginManager alloc] init];
-				// Start with iOS level account information, falls back to Native app, then web
-				//loginManager.loginBehavior = FBSDKLoginBehaviorSystemAccount;
 				NSMutableArray* Permissions = [[NSMutableArray alloc] initWithCapacity:ScopeFields.Num()];
 				for (int32 ScopeIdx = 0; ScopeIdx < ScopeFields.Num(); ScopeIdx++)
 				{
@@ -82,28 +143,31 @@ bool FOnlineIdentityFacebook::Login(int32 LocalUserNum, const FOnlineAccountCred
 					[Permissions addObject: ScopeStr];
 				}
 
-				[loginManager logInWithPermissions:Permissions
-					fromViewController:nil
-					handler: ^(FBSDKLoginManagerLoginResult* result, NSError* error)
+				FBSDKLoginConfiguration *Configuration = [[FBSDKLoginConfiguration alloc] initWithPermissions: Permissions
+																									 tracking: bIsUsingClassicLogin? FBSDKLoginTrackingEnabled : FBSDKLoginTrackingLimited];
+				
+				[loginManager logInFromViewController: nil
+					configuration: Configuration
+					completion: ^(FBSDKLoginManagerLoginResult* result, NSError* error)
 					{
-						UE_LOG_ONLINE_IDENTITY(Display, TEXT("[FBSDKLoginManager logInWithReadPermissions]"));
+						UE_LOG_ONLINE_IDENTITY(Display, TEXT("[FBSDKLoginManager logInFromViewController]"));
 						bool bSuccessfulLogin = false;
 
 						FString ErrorStr;
 						if(error)
 						{
 							ErrorStr = FString::Printf(TEXT("[%d] %s"), [error code], [error localizedDescription]);
-							UE_LOG_ONLINE_IDENTITY(Display, TEXT("[FBSDKLoginManager logInWithReadPermissions = %s]"), *ErrorStr);
+							UE_LOG_ONLINE_IDENTITY(Display, TEXT("[FBSDKLoginManager logInFromViewController = %s]"), *ErrorStr);
 
 						}
 						else if(result.isCancelled)
 						{
 							ErrorStr = LOGIN_CANCELLED;
-							UE_LOG_ONLINE_IDENTITY(Display, TEXT("[FBSDKLoginManager logInWithReadPermissions = cancelled"));
-						}						
+							UE_LOG_ONLINE_IDENTITY(Display, TEXT("[FBSDKLoginManager logInFromViewController = cancelled"));
+						}
 						else
 						{
-							UE_LOG_ONLINE_IDENTITY(Display, TEXT("[FBSDKLoginManager logInWithReadPermissions = true]"));
+							UE_LOG_ONLINE_IDENTITY(Display, TEXT("[FBSDKLoginManager logInFromViewController = true]"));
 							bSuccessfulLogin = true;
 						}
 
@@ -121,19 +185,33 @@ bool FOnlineIdentityFacebook::Login(int32 LocalUserNum, const FOnlineAccountCred
                             DeclinedPermissions.Add(permission);
                         }
 
-						const FString AccessToken([[result token] tokenString]);
 						[FIOSAsyncTask CreateTaskWithBlock : ^ bool(void)
 						{
 							// Trigger this on the game thread
 							if (bSuccessfulLogin)
 							{
-                                TSharedPtr<FOnlineSharingFacebook> Sharing = StaticCastSharedPtr<FOnlineSharingFacebook>(FacebookSubsystem->GetSharingInterface());
-                                Sharing->SetCurrentPermissions(GrantedPermissions, DeclinedPermissions);
-								Login(LocalUserNum, AccessToken);
+								TSharedPtr<FOnlineSharingFacebook> Sharing = StaticCastSharedPtr<FOnlineSharingFacebook>(FacebookSubsystem->GetSharingInterface());
+								if (bIsUsingClassicLogin)
+								{
+									const FString AccessToken([[result token] tokenString]);
+									if (Sharing)
+									{
+										Sharing->SetCurrentPermissions(GrantedPermissions, DeclinedPermissions);
+									}
+									Login(LocalUserNum, AccessToken);
+								}
+								else
+								{
+									if (Sharing)
+									{
+										Sharing->SetCurrentPermissions(GrantedPermissions, DeclinedPermissions);
+									}
+									LoginLimited(LocalUserNum);
+								}
 							}
 							else
 							{
-								OnLoginAttemptComplete(LocalUserNum, ErrorStr);
+								OnLoginAttemptComplete(LocalUserNum, false, ErrorStr);
 							}
 
 							return true;
@@ -144,14 +222,26 @@ bool FOnlineIdentityFacebook::Login(int32 LocalUserNum, const FOnlineAccountCred
 			else
 			{
 				// Skip right to attempting to use the token to query user profile
+				// or the current profile in case of limited login
 				// Could fail with an expired auth token (eg. user revoked app)
 
-				const FString AccessToken([accessToken tokenString]);
-				[FIOSAsyncTask CreateTaskWithBlock : ^ bool(void)
+				if (bIsUsingClassicLogin)
 				{
-					Login(LocalUserNum, AccessToken);
-					return true;
-				 }];
+					const FString AccessToken([FBSDKAccessToken.currentAccessToken tokenString]);
+					[FIOSAsyncTask CreateTaskWithBlock : ^ bool(void)
+					 {
+						Login(LocalUserNum, FString(AccessToken));
+						return true;
+					}];
+				}
+				else
+				{
+					[FIOSAsyncTask CreateTaskWithBlock : ^ bool(void)
+					 {
+						LoginLimited(LocalUserNum);
+						return true;
+					}];
+				}
 			}
 		}
 	);
@@ -163,43 +253,53 @@ void FOnlineIdentityFacebook::Login(int32 LocalUserNum, const FString& AccessTok
 {
 	FOnProfileRequestComplete CompletionDelegate = FOnProfileRequestComplete::CreateLambda([this](int32 LocalUserNumFromRequest, bool bWasProfileRequestSuccessful, const FString& ErrorStr)
 	{
-		FOnRequestCurrentPermissionsComplete NextCompletionDelegate = FOnRequestCurrentPermissionsComplete::CreateLambda([this](int32 LocalUserNumFromPerms, bool bWerePermsSuccessful, const TArray<FSharingPermission>& Permissions)
-		{
-			OnRequestCurrentPermissionsComplete(LocalUserNumFromPerms, bWerePermsSuccessful, Permissions);
-		});
-
-		if (bWasProfileRequestSuccessful)
-		{
-			RequestCurrentPermissions(LocalUserNumFromRequest, NextCompletionDelegate);
-		}
-		else
-		{
-			OnLoginAttemptComplete(LocalUserNumFromRequest, ErrorStr);
-		}
+		OnLoginAttemptComplete(LocalUserNumFromRequest, bWasProfileRequestSuccessful, ErrorStr);
 	});
 
 	ProfileRequest(LocalUserNum, AccessToken, ProfileFields, CompletionDelegate);
 }
 
-void FOnlineIdentityFacebook::OnRequestCurrentPermissionsComplete(int32 LocalUserNum, bool bWasSuccessful, const TArray<FSharingPermission>& NewPermissions)
+void FOnlineIdentityFacebook::LoginLimited(int32 LocalUserNum)
 {
-	FString ErrorStr;
-	if (!bWasSuccessful)
+	// Gather data from profile snapshot and store it
+	TSharedRef<FUserOnlineAccountFacebookIOS> User = MakeShared<FUserOnlineAccountFacebookIOS>(FBSDKProfile.currentProfile);
+	UserAccounts.Add(User->GetUserId()->ToString(), User);
+	UserIds.Add(LocalUserNum, User->GetUserId());
+	
+	OnLoginAttemptComplete(LocalUserNum, true, TEXT(""));
+}
+
+FString FOnlineIdentityFacebook::GetAuthToken(int32 LocalUserNum) const
+{
+	if (IsUsingClassicLogin())
 	{
-		ErrorStr = TEXT("Failure to request current sharing permissions");
+		return FOnlineIdentityFacebookCommon::GetAuthToken(LocalUserNum);
 	}
-
-	LoginStatus = bWasSuccessful ? ELoginStatus::LoggedIn : ELoginStatus::NotLoggedIn;
-	OnLoginAttemptComplete(LocalUserNum, ErrorStr);
+	FUniqueNetIdPtr UserId = GetUniquePlayerId(LocalUserNum);
+	if (UserId.IsValid())
+	{
+		TSharedPtr<FUserOnlineAccount> UserAccount = GetUserAccount(*UserId);
+		if (UserAccount.IsValid())
+		{
+			FString AuthToken;
+			if (UserAccount->GetAuthAttribute(AUTH_ATTR_ID_TOKEN, AuthToken))
+			{
+				return AuthToken;
+			}
+		}
+	}
+	return FString();
 }
 
-void FOnlineIdentityFacebook::OnLoginAttemptComplete(int32 LocalUserNum, const FString& ErrorStr)
+void FOnlineIdentityFacebook::OnLoginAttemptComplete(int32 LocalUserNum, bool bSucceeded, const FString& ErrorStr)
 {
+	LoginStatus = bSucceeded ? ELoginStatus::LoggedIn : ELoginStatus::NotLoggedIn;
 	if (LoginStatus == ELoginStatus::LoggedIn)
 	{
 		UE_LOG_ONLINE_IDENTITY(Display, TEXT("Facebook login was successful"));
 		FUniqueNetIdPtr UserId = GetUniquePlayerId(LocalUserNum);
 		check(UserId.IsValid());
+		bIsLoginInProgress = false;
 		TriggerOnLoginCompleteDelegates(LocalUserNum, true, *UserId, ErrorStr);
 		TriggerOnLoginStatusChangedDelegates(LocalUserNum, ELoginStatus::NotLoggedIn, ELoginStatus::LoggedIn, *UserId);
 	}
@@ -230,6 +330,7 @@ void FOnlineIdentityFacebook::OnLoginAttemptComplete(int32 LocalUserNum, const F
 				// remove cached user id
 				UserIds.Remove(LocalUserNum);
 
+				bIsLoginInProgress = false;
 				TriggerOnLoginCompleteDelegates(LocalUserNum, false, *UserId, NewErrorStr);
 				return true;
 			 }];
@@ -239,10 +340,8 @@ void FOnlineIdentityFacebook::OnLoginAttemptComplete(int32 LocalUserNum, const F
 
 bool FOnlineIdentityFacebook::Logout(int32 LocalUserNum)
 {
-	if ([FBSDKAccessToken currentAccessToken])
+	if (LoginStatus == ELoginStatus::LoggedIn)
 	{
-		ensure(LoginStatus == ELoginStatus::LoggedIn);
-
 		dispatch_async(dispatch_get_main_queue(),^
 		{
 			FBSDKLoginManager* loginManager = [[FBSDKLoginManager alloc] init];
@@ -275,8 +374,6 @@ bool FOnlineIdentityFacebook::Logout(int32 LocalUserNum)
 	}
 	else
 	{
-		ensure(LoginStatus == ELoginStatus::NotLoggedIn);
-
 		UE_LOG_ONLINE_IDENTITY(Warning, TEXT("No logged in user found for LocalUserNum=%d."), LocalUserNum);
 		FacebookSubsystem->ExecuteNextTick([this, LocalUserNum](){
 			TriggerOnLogoutCompleteDelegates(LocalUserNum, false);
