Update Facebook SDK for Android

diff --git a/Engine/Plugins/Online/OnlineSubsystemFacebook/Source/OnlineSubsystemFacebookAndroid_UPL.xml b/Engine/Plugins/Online/OnlineSubsystemFacebook/Source/OnlineSubsystemFacebookAndroid_UPL.xml
--- a/Engine/Plugins/Online/OnlineSubsystemFacebook/Source/OnlineSubsystemFacebookAndroid_UPL.xml
+++ b/Engine/Plugins/Online/OnlineSubsystemFacebook/Source/OnlineSubsystemFacebookAndroid_UPL.xml
@@ -15,6 +15,11 @@
 		<dumpvars/>
 	</init>
 	
+	<registerBuildSettings>
+	<insertValue value="OnlineSubsystemFacebook[bEnabled,ClientId,ClientToken,bEnableAutomaticLogging,bEnableAdvertisingId]"/>
+	<insertNewline/>
+	</registerBuildSettings>
+
 	<!-- optional files or directories to copy to Intermediate/Android/APK -->
 	<resourceCopies>
 	</resourceCopies>
@@ -41,8 +46,8 @@
 			<true>
 				<insert>
 dependencies {
-	implementation('com.facebook.android:facebook-android-sdk:15.0.1')
-	implementation('com.facebook.android:facebook-login:15.0.1')
+	implementation('com.facebook.android:facebook-android-sdk:17.0.0')
+	implementation('com.facebook.android:facebook-login:17.0.0')
 }
 				</insert>
 			</true>
diff --git a/Engine/Plugins/Online/OnlineSubsystemFacebook/Source/Private/Android/OnlineExternalUIInterfaceFacebook.cpp b/Engine/Plugins/Online/OnlineSubsystemFacebook/Source/Private/Android/OnlineExternalUIInterfaceFacebook.cpp
--- a/Engine/Plugins/Online/OnlineSubsystemFacebook/Source/Private/Android/OnlineExternalUIInterfaceFacebook.cpp
+++ b/Engine/Plugins/Online/OnlineSubsystemFacebook/Source/Private/Android/OnlineExternalUIInterfaceFacebook.cpp
@@ -1,6 +1,9 @@
 // Copyright Epic Games, Inc. All Rights Reserved.
 
 #include "OnlineExternalUIInterfaceFacebook.h"
+#include "OnlineSubsystemFacebook.h"
+#include "OnlineIdentityFacebook.h"
+#include "OnlineError.h"
 
 FOnlineExternalUIFacebook::FOnlineExternalUIFacebook(FOnlineSubsystemFacebook* InSubsystem) :
 	FOnlineExternalUIFacebookCommon(InSubsystem)
@@ -9,4 +12,29 @@ FOnlineExternalUIFacebook::FOnlineExternalUIFacebook(FOnlineSubsystemFacebook* I
 
 FOnlineExternalUIFacebook::~FOnlineExternalUIFacebook()
 {
+}
+
+bool FOnlineExternalUIFacebook::ShowLoginUI(const int ControllerIndex, bool bShowOnlineOnly, bool bShowSkipButton, const FOnLoginUIClosedDelegate& Delegate)
+{
+	IOnlineIdentityPtr OnlineIdentity = FacebookSubsystem->GetIdentityInterface();
+
+	if (FUniqueNetIdPtr UniqueNetId = OnlineIdentity->GetUniquePlayerId(ControllerIndex); UniqueNetId && UniqueNetId->IsValid())
+	{
+		Delegate.ExecuteIfBound(UniqueNetId, ControllerIndex, FOnlineError::Success());
+		return true;
+	}
+
+	TSharedPtr<FDelegateHandle> DelegateHandle = MakeShared<FDelegateHandle>();
+	*DelegateHandle = OnlineIdentity->AddOnLoginCompleteDelegate_Handle(ControllerIndex, FOnLoginCompleteDelegate::CreateLambda([this, DelegateHandle, Delegate](int32 LocalUserNum, bool bWasSuccessful, const FUniqueNetId& UserId, const FString& ErrorString)
+		{
+			FOnlineError Error(bWasSuccessful);
+			Error.SetFromErrorCode(ErrorString);
+
+			Delegate.ExecuteIfBound(UserId.IsValid()? UserId.AsShared() : FUniqueNetIdPtr(), LocalUserNum, Error);
+
+			IOnlineIdentityPtr OnlineIdentity = FacebookSubsystem->GetIdentityInterface();
+			OnlineIdentity->ClearOnLoginCompleteDelegate_Handle(LocalUserNum, *DelegateHandle);
+		}));
+	OnlineIdentity->Login(ControllerIndex, FOnlineAccountCredentials());
+	return true;
 }
\ No newline at end of file
diff --git a/Engine/Plugins/Online/OnlineSubsystemFacebook/Source/Private/Android/OnlineExternalUIInterfaceFacebook.h b/Engine/Plugins/Online/OnlineSubsystemFacebook/Source/Private/Android/OnlineExternalUIInterfaceFacebook.h
--- a/Engine/Plugins/Online/OnlineSubsystemFacebook/Source/Private/Android/OnlineExternalUIInterfaceFacebook.h
+++ b/Engine/Plugins/Online/OnlineSubsystemFacebook/Source/Private/Android/OnlineExternalUIInterfaceFacebook.h
@@ -29,6 +29,9 @@ public:
 	 * Destructor.
 	 */
 	virtual ~FOnlineExternalUIFacebook();
+
+	// IOnlineExternalUI
+	virtual bool ShowLoginUI(const int ControllerIndex, bool bShowOnlineOnly, bool bShowSkipButton, const FOnLoginUIClosedDelegate& Delegate = FOnLoginUIClosedDelegate()) override;
 };
 
 typedef TSharedPtr<FOnlineExternalUIFacebook, ESPMode::ThreadSafe> FOnlineExternalUIFacebookPtr;
diff --git a/Engine/Plugins/Online/OnlineSubsystemFacebook/Source/Private/Android/OnlineSubsystemFacebook.cpp b/Engine/Plugins/Online/OnlineSubsystemFacebook/Source/Private/Android/OnlineSubsystemFacebook.cpp
--- a/Engine/Plugins/Online/OnlineSubsystemFacebook/Source/Private/Android/OnlineSubsystemFacebook.cpp
+++ b/Engine/Plugins/Online/OnlineSubsystemFacebook/Source/Private/Android/OnlineSubsystemFacebook.cpp
@@ -2,6 +2,7 @@
 
 #include "OnlineSubsystemFacebook.h"
 #include "OnlineSubsystemFacebookPrivate.h"
+#include "OnlineExternalUIInterfaceFacebook.h"
 #include "OnlineIdentityFacebook.h"
 #include "OnlineFriendsFacebook.h"
 #include "OnlineSharingFacebook.h"
@@ -23,6 +24,7 @@ bool FOnlineSubsystemFacebook::Init()
 	FacebookFriends = MakeShared<FOnlineFriendsFacebook>(this);
 	FacebookSharing = MakeShared<FOnlineSharingFacebook>(this);
     FacebookUser = MakeShared<FOnlineUserFacebook>(this);
+	FacebookExternalUI = MakeShared<FOnlineExternalUIFacebook>(this);
 	return true;
 #else
 	return false;
