Printf format fixes.

#rnx

diff --git a/Engine/Plugins/Online/OnlineSubsystemFacebook/Source/Private/IOS/OnlineIdentityFacebook.cpp b/Engine/Plugins/Online/OnlineSubsystemFacebook/Source/Private/IOS/OnlineIdentityFacebook.cpp
--- a/Engine/Plugins/Online/OnlineSubsystemFacebook/Source/Private/IOS/OnlineIdentityFacebook.cpp
+++ b/Engine/Plugins/Online/OnlineSubsystemFacebook/Source/Private/IOS/OnlineIdentityFacebook.cpp
@@ -156,7 +156,7 @@ bool FOnlineIdentityFacebook::Login(int32 LocalUserNum, const FOnlineAccountCred
 						FString ErrorStr;
 						if(error)
 						{
-							ErrorStr = FString::Printf(TEXT("[%d] %s"), [error code], [error localizedDescription]);
+							ErrorStr = FString::Printf(TEXT("[%d] %hs"), [error code], [error localizedDescription]);
 							UE_LOG_ONLINE_IDENTITY(Display, TEXT("[FBSDKLoginManager logInFromViewController = %s]"), *ErrorStr);
 
 						}
diff --git a/Engine/Source/Runtime/HardwareSurvey/Private/HardwareSurveyModule.cpp b/Engine/Source/Runtime/HardwareSurvey/Private/HardwareSurveyModule.cpp
--- a/Engine/Source/Runtime/HardwareSurvey/Private/HardwareSurveyModule.cpp
+++ b/Engine/Source/Runtime/HardwareSurvey/Private/HardwareSurveyModule.cpp
@@ -257,7 +257,7 @@ protected:
 			HardwareStatsAttribs.Add(FAnalyticsEventAttribute(TEXT("GPUModel"), SurveyResults.RHIAdapter.AdapterName));
 #endif
 
-			Analytics->RecordEvent(*FString::Printf(TEXT("%sHardwareStats"), FPlatformProperties::IniPlatformName()), HardwareStatsAttribs);
+			Analytics->RecordEvent(*FString::Printf(TEXT("%hsHardwareStats"), FPlatformProperties::IniPlatformName()), HardwareStatsAttribs);
 
 #elif PLATFORM_DESKTOP
 
