Probably CIS fix.

diff --git a/Engine/Plugins/Online/OnlineSubsystemFacebook/Source/Private/IOS/OnlineIdentityFacebook.cpp b/Engine/Plugins/Online/OnlineSubsystemFacebook/Source/Private/IOS/OnlineIdentityFacebook.cpp
--- a/Engine/Plugins/Online/OnlineSubsystemFacebook/Source/Private/IOS/OnlineIdentityFacebook.cpp
+++ b/Engine/Plugins/Online/OnlineSubsystemFacebook/Source/Private/IOS/OnlineIdentityFacebook.cpp
@@ -156,7 +156,7 @@ bool FOnlineIdentityFacebook::Login(int32 LocalUserNum, const FOnlineAccountCred
 						FString ErrorStr;
 						if(error)
 						{
-							ErrorStr = FString::Printf(TEXT("[%d] %hs"), [error code], [error localizedDescription]);
+							ErrorStr = FString::Printf(TEXT("[%d] %s"), [error code], WCHAR_TO_TCHAR([error localizedDescription]));
 							UE_LOG_ONLINE_IDENTITY(Display, TEXT("[FBSDKLoginManager logInFromViewController = %s]"), *ErrorStr);
 
 						}
