Update facebook dependency to avoid kotlin stdlib conflict KT-55297

#rnx

diff --git a/Engine/Plugins/Online/OnlineSubsystemFacebook/Source/OnlineSubsystemFacebookAndroid_UPL.xml b/Engine/Plugins/Online/OnlineSubsystemFacebook/Source/OnlineSubsystemFacebookAndroid_UPL.xml
--- a/Engine/Plugins/Online/OnlineSubsystemFacebook/Source/OnlineSubsystemFacebookAndroid_UPL.xml
+++ b/Engine/Plugins/Online/OnlineSubsystemFacebook/Source/OnlineSubsystemFacebookAndroid_UPL.xml
@@ -46,8 +46,8 @@
 			<true>
 				<insert>
 dependencies {
-	implementation('com.facebook.android:facebook-android-sdk:17.0.0')
-	implementation('com.facebook.android:facebook-login:17.0.0')
+	implementation('com.facebook.android:facebook-android-sdk:18.0.2')
+	implementation('com.facebook.android:facebook-login:18.0.2')
 }
 				</insert>
 			</true>
