Fix IOS CIS Issues.

As part of jira UE-360487, we, removed, sdk, headers, from, all, platform, files., All, affected, files, in, the, release, branch, were, fixed, but, there, were, some, in, main, that, also, needed, to, specifically, includes, they, need.

This CL just adds the missing headers / foward declarations.

#Synced-CL 49790678

diff --git a/Engine/Plugins/Online/OnlineSubsystemApple/Source/Private/OnlineExternalUIInterfaceApple.h b/Engine/Plugins/Online/OnlineSubsystemApple/Source/Private/OnlineExternalUIInterfaceApple.h
--- a/Engine/Plugins/Online/OnlineSubsystemApple/Source/Private/OnlineExternalUIInterfaceApple.h
+++ b/Engine/Plugins/Online/OnlineSubsystemApple/Source/Private/OnlineExternalUIInterfaceApple.h
@@ -5,6 +5,10 @@
 #include "Interfaces/OnlineExternalUIInterface.h"
 #include "OnlineSubsystemApplePackage.h"
 
+@class NSData; 
+@class NSPersonNameComponents; 
+@class NSError;
+
 class FOnlineSubsystemApple;
 
 class FOnlineExternalUIApple : public IOnlineExternalUI
diff --git a/Engine/Plugins/Online/OnlineSubsystemFacebook/Source/Private/IOS/FacebookHelper.cpp b/Engine/Plugins/Online/OnlineSubsystemFacebook/Source/Private/IOS/FacebookHelper.cpp
--- a/Engine/Plugins/Online/OnlineSubsystemFacebook/Source/Private/IOS/FacebookHelper.cpp
+++ b/Engine/Plugins/Online/OnlineSubsystemFacebook/Source/Private/IOS/FacebookHelper.cpp
@@ -1,6 +1,8 @@
 // Copyright Epic Games, Inc. All Rights Reserved.
 
 #include "FacebookHelper.h"
+#include "Apple/AppleAsyncTask.h"
+#include "Apple/AppleStringUtils.h"
 #include "OnlineSubsystemFacebookPrivate.h"
 #include "Interfaces/OnlineIdentityInterface.h"
 
@@ -11,6 +13,7 @@ THIRD_PARTY_INCLUDES_START
 #import <FBSDKLoginKit/FBSDKLoginKit.h>
 THIRD_PARTY_INCLUDES_END
 
+
 @implementation FFacebookHelper
 
 - (id)initWithOwner:(TWeakPtr<FIOSFacebookNotificationDelegate>) InOwner
diff --git a/Engine/Plugins/Online/OnlineSubsystemFacebook/Source/Private/IOS/OnlineIdentityFacebook.cpp b/Engine/Plugins/Online/OnlineSubsystemFacebook/Source/Private/IOS/OnlineIdentityFacebook.cpp
--- a/Engine/Plugins/Online/OnlineSubsystemFacebook/Source/Private/IOS/OnlineIdentityFacebook.cpp
+++ b/Engine/Plugins/Online/OnlineSubsystemFacebook/Source/Private/IOS/OnlineIdentityFacebook.cpp
@@ -2,6 +2,8 @@
 
 // Module includes
 #include "OnlineIdentityFacebook.h"
+#include "Apple/AppleAsyncTask.h"
+#include "Apple/AppleStringUtils.h"
 #include "OnlineSharingFacebook.h"
 #include "OnlineSubsystemFacebookPrivate.h"
 #include "Interfaces/OnlineSharingInterface.h"
diff --git a/Engine/Plugins/Online/OnlineSubsystemFacebook/Source/Private/IOS/OnlineSharingFacebook.cpp b/Engine/Plugins/Online/OnlineSubsystemFacebook/Source/Private/IOS/OnlineSharingFacebook.cpp
--- a/Engine/Plugins/Online/OnlineSubsystemFacebook/Source/Private/IOS/OnlineSharingFacebook.cpp
+++ b/Engine/Plugins/Online/OnlineSubsystemFacebook/Source/Private/IOS/OnlineSharingFacebook.cpp
@@ -3,6 +3,8 @@
 
 // Module includes
 #include "OnlineSharingFacebook.h"
+#include "Apple/AppleAsyncTask.h"
+#include "Apple/AppleStringUtils.h"
 #include "OnlineSubsystemFacebookPrivate.h"
 #include "OnlineIdentityFacebook.h"
 
diff --git a/Engine/Plugins/Online/OnlineSubsystemFacebook/Source/Private/IOS/OnlineSubsystemFacebook.cpp b/Engine/Plugins/Online/OnlineSubsystemFacebook/Source/Private/IOS/OnlineSubsystemFacebook.cpp
--- a/Engine/Plugins/Online/OnlineSubsystemFacebook/Source/Private/IOS/OnlineSubsystemFacebook.cpp
+++ b/Engine/Plugins/Online/OnlineSubsystemFacebook/Source/Private/IOS/OnlineSubsystemFacebook.cpp
@@ -1,6 +1,8 @@
 // Copyright Epic Games, Inc. All Rights Reserved.
 
 #include "OnlineSubsystemFacebook.h"
+#include "Apple/AppleAsyncTask.h"
+#include "Apple/AppleStringUtils.h"
 #include "OnlineSubsystemFacebookPrivate.h"
 
 #include "Misc/CoreDelegates.h"
diff --git a/Engine/Plugins/Runtime/Advertising/IOSTapJoy/Source/IOSTapJoy/Private/IOSTapJoy.cpp b/Engine/Plugins/Runtime/Advertising/IOSTapJoy/Source/IOSTapJoy/Private/IOSTapJoy.cpp
--- a/Engine/Plugins/Runtime/Advertising/IOSTapJoy/Source/IOSTapJoy/Private/IOSTapJoy.cpp
+++ b/Engine/Plugins/Runtime/Advertising/IOSTapJoy/Source/IOSTapJoy/Private/IOSTapJoy.cpp
@@ -1,6 +1,7 @@
 // Copyright Epic Games, Inc. All Rights Reserved.
 
 #include "CoreMinimal.h"
+#include "Apple/AppleStringUtils.h"
 #include "Interfaces/IAdvertisingProvider.h"
 #include "IOS/IOSAppDelegate.h"
 #include "Modules/ModuleInterface.h"
diff --git a/Engine/Plugins/Runtime/Analytics/Adjust/Source/IOSAdjust/Private/IOSAdjust.cpp b/Engine/Plugins/Runtime/Analytics/Adjust/Source/IOSAdjust/Private/IOSAdjust.cpp
--- a/Engine/Plugins/Runtime/Analytics/Adjust/Source/IOSAdjust/Private/IOSAdjust.cpp
+++ b/Engine/Plugins/Runtime/Analytics/Adjust/Source/IOSAdjust/Private/IOSAdjust.cpp
@@ -1,6 +1,7 @@
 // Copyright Epic Games, Inc. All Rights Reserved.
 
 #include "IOSAdjust.h"
+#include "Apple/AppleStringUtils.h"
 #include "Analytics.h"
 #include "IOSAdjustProvider.h"
 #include "Misc/Paths.h"
diff --git a/Engine/Plugins/Runtime/LocationServicesIOSImpl/Source/LocationServicesIOSImpl/Private/LocationServicesIOSImpl.cpp b/Engine/Plugins/Runtime/LocationServicesIOSImpl/Source/LocationServicesIOSImpl/Private/LocationServicesIOSImpl.cpp
--- a/Engine/Plugins/Runtime/LocationServicesIOSImpl/Source/LocationServicesIOSImpl/Private/LocationServicesIOSImpl.cpp
+++ b/Engine/Plugins/Runtime/LocationServicesIOSImpl/Source/LocationServicesIOSImpl/Private/LocationServicesIOSImpl.cpp
@@ -2,6 +2,7 @@
 
 #include "LocationServicesIOSImpl.h"
 #include <CoreLocation/CoreLocation.h>
+#import <UIKit/UIKit.h>
 
 DEFINE_LOG_CATEGORY(LogLocationServicesIOS);
 
